Fix random crash on shutdown in OCIO editor module

diff --git a/Engine/Plugins/Compositing/OpenColorIO/Source/OpenColorIOEditor/Private/OpenColorIOEditorModule.cpp b/Engine/Plugins/Compositing/OpenColorIO/Source/OpenColorIOEditor/Private/OpenColorIOEditorModule.cpp
--- a/Engine/Plugins/Compositing/OpenColorIO/Source/OpenColorIOEditor/Private/OpenColorIOEditorModule.cpp
+++ b/Engine/Plugins/Compositing/OpenColorIO/Source/OpenColorIOEditor/Private/OpenColorIOEditorModule.cpp
@@ -80,9 +80,12 @@ public:
 
 	void UnregisterCustomizations()
 	{
-		FPropertyEditorModule& PropertyModule = FModuleManager::LoadModuleChecked<FPropertyEditorModule>("PropertyEditor");
-		PropertyModule.UnregisterCustomPropertyTypeLayout(FOpenColorIOColorSpace::StaticStruct()->GetFName());
-		PropertyModule.UnregisterCustomPropertyTypeLayout(FOpenColorIOColorConversionSettings::StaticStruct()->GetFName());
+		if (UObjectInitialized() == true)
+		{
+			FPropertyEditorModule& PropertyModule = FModuleManager::LoadModuleChecked<FPropertyEditorModule>("PropertyEditor");
+			PropertyModule.UnregisterCustomPropertyTypeLayout(FOpenColorIOColorSpace::StaticStruct()->GetFName());
+			PropertyModule.UnregisterCustomPropertyTypeLayout(FOpenColorIOColorConversionSettings::StaticStruct()->GetFName());
+		}
 	}
 
 	void OnWorldInit(UWorld* InWorld, const UWorld::InitializationValues InInitializationValues)
