Composure
- Adding back bAlwaysPersistRenderingState in scene capture setup for bloom and tonemapping passes. Made them unusable and broke SimpleComposure sample.

OCIO
- Fix for customization not working for certain uproperties

diff --git a/Engine/Plugins/Compositing/Composure/Source/Composure/Private/ComposureLensBloomPass.cpp b/Engine/Plugins/Compositing/Composure/Source/Composure/Private/ComposureLensBloomPass.cpp
--- a/Engine/Plugins/Compositing/Composure/Source/Composure/Private/ComposureLensBloomPass.cpp
+++ b/Engine/Plugins/Compositing/Composure/Source/Composure/Private/ComposureLensBloomPass.cpp
@@ -66,6 +66,10 @@ void UComposureLensBloomPass::BloomToRenderTarget()
 
 	SceneCapture->ProfilingEventName = TEXT("ComposureLensBloomPass");
 
+	// OverrideBlendableSettings() will do nothing (see UMaterialInterface::OverrideBlendableSettings) 
+	// with these materials unless there is a ViewState from the capture component (see USceneCaptureComponent::GetViewState)
+	TGuardValue<bool> ViewStateGuard(SceneCapture->bAlwaysPersistRenderingState, true);
+
 	// Update the render target output.
 	SceneCapture->CaptureScene();
 }
diff --git a/Engine/Plugins/Compositing/Composure/Source/Composure/Private/ComposureTonemapperPass.cpp b/Engine/Plugins/Compositing/Composure/Source/Composure/Private/ComposureTonemapperPass.cpp
--- a/Engine/Plugins/Compositing/Composure/Source/Composure/Private/ComposureTonemapperPass.cpp
+++ b/Engine/Plugins/Compositing/Composure/Source/Composure/Private/ComposureTonemapperPass.cpp
@@ -63,6 +63,10 @@ void UComposureTonemapperPass::TonemapToRenderTarget()
 
 	SceneCapture->ProfilingEventName = TEXT("ComposureTonemapperPass");
 
+	// OverrideBlendableSettings() will do nothing (see UMaterialInterface::OverrideBlendableSettings) 
+	// with these materials unless there is a ViewState from the capture component (see USceneCaptureComponent::GetViewState)
+	TGuardValue<bool> ViewStateGuard(SceneCapture->bAlwaysPersistRenderingState, true);
+
 	// Update the render target output.
 	SceneCapture->CaptureScene();
 }
diff --git a/Engine/Plugins/Compositing/OpenColorIO/Source/OpenColorIOEditor/Private/OpenColorIOColorSpaceCustomization.cpp b/Engine/Plugins/Compositing/OpenColorIO/Source/OpenColorIOEditor/Private/OpenColorIOColorSpaceCustomization.cpp
--- a/Engine/Plugins/Compositing/OpenColorIO/Source/OpenColorIOEditor/Private/OpenColorIOColorSpaceCustomization.cpp
+++ b/Engine/Plugins/Compositing/OpenColorIO/Source/OpenColorIOEditor/Private/OpenColorIOColorSpaceCustomization.cpp
@@ -22,7 +22,6 @@ void FOpenColorIOColorSpaceCustomization::CustomizeHeader(TSharedRef<IPropertyHa
 	if (ColorSpaceProperty->GetNumPerObjectValues() == 1 && ColorSpaceProperty->IsValidHandle())
 	{
 		UProperty* Property = ColorSpaceProperty->GetProperty();
-		TSharedPtr<IPropertyHandle> ParentArrayHandle = InPropertyHandle->GetParentHandle();
 		check(Property && Cast<UStructProperty>(Property) && Cast<UStructProperty>(Property)->Struct && Cast<UStructProperty>(Property)->Struct->IsChildOf(FOpenColorIOColorSpace::StaticStruct()));
 
 		TArray<void*> RawData;
@@ -37,6 +36,9 @@ void FOpenColorIOColorSpaceCustomization::CustomizeHeader(TSharedRef<IPropertyHa
 		// Get the ConfigurationFile to read color space from
 		{
 			static FName NAME_ConfigFile = TEXT("OCIOConfigFile");
+
+			//Verify if we're in an array before accessing the property directly. Need to go deeper for array properties.
+			TSharedPtr<IPropertyHandle> ParentArrayHandle = InPropertyHandle->GetParentHandle();
 			TSharedPtr<IPropertyHandleArray> ParentHandleAsArray = ParentArrayHandle->AsArray();
 			if (ParentHandleAsArray.IsValid())
 			{
@@ -51,6 +53,18 @@ void FOpenColorIOColorSpaceCustomization::CustomizeHeader(TSharedRef<IPropertyHa
 					}
 				}
 			}
+			else
+			{
+				if (ColorSpaceProperty->HasMetaData(NAME_ConfigFile))
+				{
+					const FString& ConfigFileVariableName = ColorSpaceProperty->GetMetaData(NAME_ConfigFile);
+					TSharedPtr<IPropertyHandle> ClassHandle = ColorSpaceProperty->GetParentHandle();
+					if (ClassHandle.IsValid() && !ConfigFileVariableName.IsEmpty())
+					{
+						ConfigurationFileProperty = ClassHandle->GetChildHandle(*ConfigFileVariableName);
+					}
+				}
+			}
 		}
 
 
