- Fixing path management of ocio asset to be relative to the project instead of the binary
- Updating EngineTest config to add required ocio files when being distilled

diff --git a/Engine/Plugins/Compositing/OpenColorIO/Source/OpenColorIO/Private/OpenColorIOConfiguration.cpp b/Engine/Plugins/Compositing/OpenColorIO/Source/OpenColorIO/Private/OpenColorIOConfiguration.cpp
--- a/Engine/Plugins/Compositing/OpenColorIO/Source/OpenColorIO/Private/OpenColorIOConfiguration.cpp
+++ b/Engine/Plugins/Compositing/OpenColorIO/Source/OpenColorIO/Private/OpenColorIOConfiguration.cpp
@@ -187,10 +187,21 @@ void UOpenColorIOConfiguration::LoadConfigurationFile()
 		{
 			LoadedConfig.reset();
 
-			OCIO_NAMESPACE::ConstConfigRcPtr NewConfig = OCIO_NAMESPACE::Config::CreateFromFile(StringCast<ANSICHAR>(*ConfigurationFile.FilePath).Get());
+			FString FullPath;
+			if (!FPaths::IsRelative(ConfigurationFile.FilePath))
+			{
+				FullPath = ConfigurationFile.FilePath;
+			}
+			else
+			{
+				const FString AbsoluteGameDir = FPaths::ConvertRelativePathToFull(FPaths::ProjectDir());
+				FullPath = FPaths::ConvertRelativePathToFull(FPaths::Combine(AbsoluteGameDir, ConfigurationFile.FilePath));
+			}
+
+			OCIO_NAMESPACE::ConstConfigRcPtr NewConfig = OCIO_NAMESPACE::Config::CreateFromFile(StringCast<ANSICHAR>(*FullPath).Get());
 			if (NewConfig)
 			{
-				UE_LOG(LogOpenColorIO, Verbose, TEXT("Loaded OCIO configuration file %s"), *ConfigurationFile.FilePath);
+				UE_LOG(LogOpenColorIO, Verbose, TEXT("Loaded OCIO configuration file %s"), *FullPath);
 				LoadedConfig = NewConfig;
 			}
 			else
diff --git a/Engine/Plugins/Compositing/OpenColorIO/Source/OpenColorIO/Public/OpenColorIOConfiguration.h b/Engine/Plugins/Compositing/OpenColorIO/Source/OpenColorIO/Public/OpenColorIOConfiguration.h
--- a/Engine/Plugins/Compositing/OpenColorIO/Source/OpenColorIO/Public/OpenColorIOConfiguration.h
+++ b/Engine/Plugins/Compositing/OpenColorIO/Source/OpenColorIO/Public/OpenColorIOConfiguration.h
@@ -61,7 +61,7 @@ private:
 
 public:
 
-	UPROPERTY(BlueprintReadWrite, EditAnywhere, Category = "Config", meta = (FilePathFilter = "ocio"))
+	UPROPERTY(BlueprintReadWrite, EditAnywhere, Category = "Config", meta = (FilePathFilter = "ocio", RelativeToGameDir))
 	FFilePath ConfigurationFile;
 
 	UPROPERTY(BlueprintReadWrite, EditAnywhere, Category = "ColorSpace", meta=(OCIOConfigFile="ConfigurationFile"))
diff --git a/Engine/Plugins/Compositing/OpenColorIO/Source/OpenColorIOEditor/Private/OpenColorIOColorSpaceCustomization.cpp b/Engine/Plugins/Compositing/OpenColorIO/Source/OpenColorIOEditor/Private/OpenColorIOColorSpaceCustomization.cpp
--- a/Engine/Plugins/Compositing/OpenColorIO/Source/OpenColorIOEditor/Private/OpenColorIOColorSpaceCustomization.cpp
+++ b/Engine/Plugins/Compositing/OpenColorIO/Source/OpenColorIOEditor/Private/OpenColorIOColorSpaceCustomization.cpp
@@ -108,7 +108,18 @@ bool FOpenColorIOColorSpaceCustomization::LoadConfigurationFile(const FFilePath&
 	try
 #endif
 	{
-		CachedConfigFile = OCIO_NAMESPACE::Config::CreateFromFile(StringCast<ANSICHAR>(*InFilePath.FilePath).Get());
+		FString FullPath;
+		if (!FPaths::IsRelative(InFilePath.FilePath))
+		{
+			FullPath = InFilePath.FilePath;
+		}
+		else
+		{
+			const FString AbsoluteGameDir = FPaths::ConvertRelativePathToFull(FPaths::ProjectDir());
+			FullPath = FPaths::ConvertRelativePathToFull(FPaths::Combine(AbsoluteGameDir, InFilePath.FilePath));
+		}
+
+		CachedConfigFile = OCIO_NAMESPACE::Config::CreateFromFile(StringCast<ANSICHAR>(*FullPath).Get());
 		if (!CachedConfigFile)
 		{
 			return false;
