- Removed opencolorio errors when building from buildmachines on platform without the library.

diff --git a/Engine/Plugins/Compositing/OpenColorIO/Source/OpenColorIO/Private/OpenColorIOColorTransform.cpp b/Engine/Plugins/Compositing/OpenColorIO/Source/OpenColorIO/Private/OpenColorIOColorTransform.cpp
--- a/Engine/Plugins/Compositing/OpenColorIO/Source/OpenColorIO/Private/OpenColorIOColorTransform.cpp
+++ b/Engine/Plugins/Compositing/OpenColorIO/Source/OpenColorIO/Private/OpenColorIOColorTransform.cpp
@@ -416,7 +416,11 @@ bool UOpenColorIOColorTransform::UpdateShaderInfo(FString& OutShaderCodeHash, FS
 
 	return false;
 #else
-	UE_LOG(LogOpenColorIO, Error, TEXT("Can't update shader, OCIO library isn't present."));
+	//Avoid triggering errors when building maps on build machine.
+	if (!GIsBuildMachine)
+	{
+		UE_LOG(LogOpenColorIO, Error, TEXT("Can't update shader, OCIO library isn't present."));
+	}
 	return false;
 #endif //WITH_OCIO
 #else
