- Fixed bug with DDC key generation for ocio shader and lut
 * Based on material shadermap, it now uses dependant shader sources hash in the key and serializationHistory
 * OCIO library version is also used

diff --git a/Engine/Plugins/Compositing/OpenColorIO/Source/OpenColorIO/Private/OpenColorIOColorTransform.cpp b/Engine/Plugins/Compositing/OpenColorIO/Source/OpenColorIO/Private/OpenColorIOColorTransform.cpp
--- a/Engine/Plugins/Compositing/OpenColorIO/Source/OpenColorIO/Private/OpenColorIOColorTransform.cpp
+++ b/Engine/Plugins/Compositing/OpenColorIO/Source/OpenColorIO/Private/OpenColorIOColorTransform.cpp
@@ -103,7 +103,14 @@ void UOpenColorIOColorTransform::ProcessSerializedShaderMaps(UOpenColorIOColorTr
 void UOpenColorIOColorTransform::GetOpenColorIOLUTKeyGuid(const FString& InLutIdentifier, FGuid& OutLutGuid)
 {
 #if WITH_EDITOR
-	const FString DDCKey = FDerivedDataCacheInterface::BuildCacheKey(TEXT("OCIOLUT"), OPENCOLORIO_DERIVEDDATA_VER, *InLutIdentifier);
+	FString DDCKey = FDerivedDataCacheInterface::BuildCacheKey(TEXT("OCIOLUT"), OPENCOLORIO_DERIVEDDATA_VER, *InLutIdentifier);
+
+#if WITH_OCIO
+	//Keep library version in the DDC key to invalidate it once we move to a new library
+	DDCKey += TEXT("OCIOVersion");
+	DDCKey += TEXT(OCIO_VERSION);
+#endif //WITH_OCIO
+
 	const uint32 KeyLength = DDCKey.Len() * sizeof(DDCKey[0]);
 	uint32 Hash[5];
 	FSHA1::HashBuffer(*DDCKey, KeyLength, reinterpret_cast<uint8*>(Hash));
diff --git a/Engine/Plugins/Compositing/OpenColorIO/Source/OpenColorIO/Private/OpenColorIODerivedDataVersion.h b/Engine/Plugins/Compositing/OpenColorIO/Source/OpenColorIO/Private/OpenColorIODerivedDataVersion.h
--- a/Engine/Plugins/Compositing/OpenColorIO/Source/OpenColorIO/Private/OpenColorIODerivedDataVersion.h
+++ b/Engine/Plugins/Compositing/OpenColorIO/Source/OpenColorIO/Private/OpenColorIODerivedDataVersion.h
@@ -10,6 +10,4 @@ OpenColorIOShaderDerivedDataVersion.h: Derived data version for OpenColorIO shad
 
 // In case of merge conflicts with DDC versions, you MUST generate a new GUID and set this new
 // guid as version
-#define OPENCOLORIO_DERIVEDDATA_VER		TEXT("E26B75AD1FD94D3F82A485887BE3CCAD")
-
-													  
+#define OPENCOLORIO_DERIVEDDATA_VER		TEXT("A0CDC343C4844554AA76FCA6A70C594E")
\ No newline at end of file
diff --git a/Engine/Plugins/Compositing/OpenColorIO/Source/OpenColorIO/Private/OpenColorIOShaderMap.cpp b/Engine/Plugins/Compositing/OpenColorIO/Source/OpenColorIO/Private/OpenColorIOShaderMap.cpp
--- a/Engine/Plugins/Compositing/OpenColorIO/Source/OpenColorIO/Private/OpenColorIOShaderMap.cpp
+++ b/Engine/Plugins/Compositing/OpenColorIO/Source/OpenColorIO/Private/OpenColorIOShaderMap.cpp
@@ -18,6 +18,11 @@
 	#include "DerivedDataCacheInterface.h"
 	#include "Interfaces/ITargetPlatformManagerModule.h"
 	#include "TickableEditorObject.h"
+
+#if WITH_OCIO
+	#include "OpenColorIO/OpenColorIO.h"
+#endif //WITH_OCIO
+
 #endif
 
 
@@ -75,6 +80,7 @@ void FOpenColorIOShaderMapId::Serialize(FArchive& Ar)
 
 	Ar << ShaderCodeHash;
 	Ar << (int32&)FeatureLevel;
+	Ar << ShaderTypeDependencies;
 }
 
 /** Hashes the color transform specific part of this shader map Id. */
@@ -97,23 +103,70 @@ void FOpenColorIOShaderMapId::GetOpenColorIOHash(FSHAHash& OutHash) const
 */
 bool FOpenColorIOShaderMapId::operator==(const FOpenColorIOShaderMapId& InReferenceSet) const
 {
-	if (ShaderCodeHash != InReferenceSet.ShaderCodeHash
+	if (  ShaderCodeHash != InReferenceSet.ShaderCodeHash
 		|| FeatureLevel != InReferenceSet.FeatureLevel)
 	{
 		return false;
 	}
 
+	if (ShaderTypeDependencies.Num() != InReferenceSet.ShaderTypeDependencies.Num())
+	{
+		return false;
+	}
+
+	for (int32 ShaderIndex = 0; ShaderIndex < ShaderTypeDependencies.Num(); ShaderIndex++)
+	{
+		const FShaderTypeDependency& ShaderTypeDependency = ShaderTypeDependencies[ShaderIndex];
+
+		if (ShaderTypeDependency.ShaderType != ShaderTypeDependencies[ShaderIndex].ShaderType
+			|| ShaderTypeDependency.SourceHash != ShaderTypeDependencies[ShaderIndex].SourceHash)
+		{
+			return false;
+		}
+	}
+
 	return true;
 }
 
 void FOpenColorIOShaderMapId::AppendKeyString(FString& OutKeyString) const
 {
+#if WITH_EDITOR
 	OutKeyString += ShaderCodeHash;
 	OutKeyString += TEXT("_");
 
 	FString FeatureLevelString;
 	GetFeatureLevelName(FeatureLevel, FeatureLevelString);
-	OutKeyString += FeatureLevelString + TEXT("_");
+
+	TMap<const TCHAR*, FCachedUniformBufferDeclaration> ReferencedUniformBuffers;
+
+	// Add the inputs for any shaders that are stored inline in the shader map
+	for (const FShaderTypeDependency& ShaderTypeDependency : ShaderTypeDependencies)
+	{
+		OutKeyString += TEXT("_");
+		OutKeyString += ShaderTypeDependency.ShaderType->GetName();
+		OutKeyString += ShaderTypeDependency.SourceHash.ToString();
+		ShaderTypeDependency.ShaderType->GetSerializationHistory().AppendKeyString(OutKeyString);
+
+		const TMap<const TCHAR*, FCachedUniformBufferDeclaration>& ReferencedUniformBufferStructsCache = ShaderTypeDependency.ShaderType->GetReferencedUniformBufferStructsCache();
+
+		for (TMap<const TCHAR*, FCachedUniformBufferDeclaration>::TConstIterator It(ReferencedUniformBufferStructsCache); It; ++It)
+		{
+			ReferencedUniformBuffers.Add(It.Key(), It.Value());
+		}
+	}
+
+	{
+		TArray<uint8> TempData;
+		FSerializationHistory SerializationHistory;
+		FMemoryWriter Ar(TempData, true);
+		FShaderSaveArchive SaveArchive(Ar, SerializationHistory);
+
+		// Save uniform buffer member info so we can detect when layout has changed
+		SerializeUniformBufferInfo(SaveArchive, ReferencedUniformBuffers);
+
+		SerializationHistory.AppendKeyString(OutKeyString);
+	}
+#endif //WITH_EDITOR
 }
 
 /**
@@ -212,7 +265,12 @@ FOpenColorIOShaderMap* FOpenColorIOShaderMap::FindId(const FOpenColorIOShaderMap
 
 void OpenColorIOShaderMapAppendKeyString(EShaderPlatform InPlatform, FString& OutKeyString)
 {
-	// does nothing at the moment
+#if WITH_EDITOR && WITH_OCIO
+	//Keep library version in the DDC key to invalidate it once we move to a new library
+	OutKeyString += TEXT("OCIOVersion");
+	OutKeyString += TEXT(OCIO_VERSION);
+	OutKeyString += TEXT("_");
+#endif //WITH_EDITOR && WITH_OCIO
 }
 
 /** Creates a string key for the derived data cache given a shader map id. */
@@ -222,6 +280,7 @@ static FString GetOpenColorIOShaderMapKeyString(const FOpenColorIOShaderMapId& I
 	const FName Format = LegacyShaderPlatformToShaderFormat(InPlatform);
 	FString ShaderMapKeyString = Format.ToString() + TEXT("_") + FString(FString::FromInt(GetTargetPlatformManagerRef().ShaderFormatVersion(Format))) + TEXT("_");
 	OpenColorIOShaderMapAppendKeyString(InPlatform, ShaderMapKeyString);
+	ShaderMapAppendKeyString(InPlatform, ShaderMapKeyString);
 	InShaderMapId.AppendKeyString(ShaderMapKeyString);
 	return FDerivedDataCacheInterface::BuildCacheKey(TEXT("OCIOSM"), OPENCOLORIO_DERIVEDDATA_VER, *ShaderMapKeyString);
 #else
@@ -356,6 +415,9 @@ void FOpenColorIOShaderMap::Compile(FOpenColorIOTransformResource* InColorTransf
 				FOpenColorIOShaderType* ShaderType = ShaderTypeIt->GetOpenColorIOShaderType();
 				if (ShaderType && ShouldCacheOpenColorIOShader(ShaderType, InPlatform, InColorTransform))
 				{
+					// Verify that the shader map Id contains inputs for any shaders that will be put into this shader map
+					check(InShaderMapId.ContainsShaderType(ShaderType));
+					
 					// Compile this OpenColorIO shader .
 					TArray<FString> ShaderErrors;
   
diff --git a/Engine/Plugins/Compositing/OpenColorIO/Source/OpenColorIO/Private/OpenColorIOShared.cpp b/Engine/Plugins/Compositing/OpenColorIO/Source/OpenColorIO/Private/OpenColorIOShared.cpp
--- a/Engine/Plugins/Compositing/OpenColorIO/Source/OpenColorIO/Private/OpenColorIOShared.cpp
+++ b/Engine/Plugins/Compositing/OpenColorIO/Source/OpenColorIO/Private/OpenColorIOShared.cpp
@@ -84,6 +84,8 @@ void FOpenColorIOTransformResource::GetDependentShaderTypes(EShaderPlatform InPl
 			OutShaderTypes.Add(ShaderType);
 		}
 	}
+
+	OutShaderTypes.Sort(FCompareShaderTypes());
 }
 
 OPENCOLORIO_API void FOpenColorIOTransformResource::GetShaderMapId(EShaderPlatform InPlatform, FOpenColorIOShaderMapId& OutId) const
@@ -96,8 +98,10 @@ OPENCOLORIO_API void FOpenColorIOTransformResource::GetShaderMapId(EShaderPlatfo
 	{
 		TArray<FShaderType*> ShaderTypes;
 		GetDependentShaderTypes(InPlatform, ShaderTypes);
+
 		OutId.FeatureLevel = GetFeatureLevel();
 		OutId.ShaderCodeHash = ShaderCodeHash;
+		OutId.SetShaderDependencies(ShaderTypes, InPlatform);
 	}
 }
 
@@ -217,9 +221,9 @@ OPENCOLORIO_API  bool FOpenColorIOTransformResource::IsCompilationFinished() con
 
 bool FOpenColorIOTransformResource::CacheShaders(EShaderPlatform InPlatform, bool bApplyCompletedShaderMapForRendering, bool bSynchronous)
 {
-	FOpenColorIOShaderMapId NoStaticParametersId;
-	GetShaderMapId(InPlatform, NoStaticParametersId);
-	return CacheShaders(NoStaticParametersId, InPlatform, bApplyCompletedShaderMapForRendering, bSynchronous);
+	FOpenColorIOShaderMapId ResourceShaderMapId;
+	GetShaderMapId(InPlatform, ResourceShaderMapId);
+	return CacheShaders(ResourceShaderMapId, InPlatform, bApplyCompletedShaderMapForRendering, bSynchronous);
 }
 
 bool FOpenColorIOTransformResource::CacheShaders(const FOpenColorIOShaderMapId& InShaderMapId, EShaderPlatform InPlatform, bool bApplyCompletedShaderMapForRendering, bool bSynchronous)
@@ -428,3 +432,33 @@ bool FOpenColorIOTransformResource::BeginCompileShaderMap(const FOpenColorIOShad
 	return false;
 #endif
 }
+
+void FOpenColorIOShaderMapId::SetShaderDependencies(const TArray<FShaderType*>& InShaderTypes, EShaderPlatform InShaderPlatform)
+{
+	if (!FPlatformProperties::RequiresCookedData())
+	{
+		for (FShaderType* ShaderType : InShaderTypes)
+		{
+			if (ShaderType != nullptr)
+			{
+				FShaderTypeDependency Dependency;
+				Dependency.ShaderType = ShaderType;
+				Dependency.SourceHash = ShaderType->GetSourceHash(InShaderPlatform);
+				ShaderTypeDependencies.Add(Dependency);
+			}
+		}
+	}
+}
+
+bool FOpenColorIOShaderMapId::ContainsShaderType(const FShaderType* ShaderType) const
+{
+	for (int32 TypeIndex = 0; TypeIndex < ShaderTypeDependencies.Num(); TypeIndex++)
+	{
+		if (ShaderTypeDependencies[TypeIndex].ShaderType == ShaderType)
+		{
+			return true;
+		}
+	}
+
+	return false;
+}
diff --git a/Engine/Plugins/Compositing/OpenColorIO/Source/OpenColorIO/Public/OpenColorIOShared.h b/Engine/Plugins/Compositing/OpenColorIO/Source/OpenColorIO/Public/OpenColorIOShared.h
--- a/Engine/Plugins/Compositing/OpenColorIO/Source/OpenColorIO/Public/OpenColorIOShared.h
+++ b/Engine/Plugins/Compositing/OpenColorIO/Source/OpenColorIO/Public/OpenColorIOShared.h
@@ -46,10 +46,14 @@ public:
 	ERHIFeatureLevel::Type FeatureLevel;
 
 	/**
-	 * The hash computed from the shader code generated.
+	 * The hash computed from the shader code generated by the library.
 	 */
 	FString ShaderCodeHash;
+	
+	/** Shader types of shaders that are inlined in this shader map in the DDC. */
+	TArray<FShaderTypeDependency> ShaderTypeDependencies;
 
+public:
 	FOpenColorIOShaderMapId()
 		: FeatureLevel(ERHIFeatureLevel::SM4)
 	{ }
@@ -57,6 +61,8 @@ public:
 	~FOpenColorIOShaderMapId()
 	{ }
 
+	void SetShaderDependencies(const TArray<FShaderType*>& InShaderTypes, EShaderPlatform InShaderPlatform);
+
 	void Serialize(FArchive& Ar);
 
 	friend uint32 GetTypeHash(const FOpenColorIOShaderMapId& Ref)
@@ -66,7 +72,7 @@ public:
 
 	SIZE_T GetSizeBytes() const
 	{
-		return sizeof(*this);
+		return sizeof(*this) + ShaderTypeDependencies.GetAllocatedSize();
 	}
 
 	/** Hashes the color transform-specific part of this shader map Id. */
@@ -87,6 +93,9 @@ public:
 
 	/** Appends string representations of this Id to a key string. */
 	void AppendKeyString(FString& OutKeyString) const;
+
+	/** Returns true if the requested shader type is a dependency of this shader map Id. */
+	bool ContainsShaderType(const FShaderType* ShaderType) const;
 };
 
 
