Fixing bad merge from Dev-Enterprise in ocio

diff --git a/Engine/Plugins/Compositing/OpenColorIO/Source/OpenColorIO/Private/OpenColorIOShaderMap.cpp b/Engine/Plugins/Compositing/OpenColorIO/Source/OpenColorIO/Private/OpenColorIOShaderMap.cpp
--- a/Engine/Plugins/Compositing/OpenColorIO/Source/OpenColorIO/Private/OpenColorIOShaderMap.cpp
+++ b/Engine/Plugins/Compositing/OpenColorIO/Source/OpenColorIO/Private/OpenColorIOShaderMap.cpp
@@ -118,8 +118,8 @@ bool FOpenColorIOShaderMapId::operator==(const FOpenColorIOShaderMapId& InRefere
 	{
 		const FShaderTypeDependency& ShaderTypeDependency = ShaderTypeDependencies[ShaderIndex];
 
-		if (ShaderTypeDependency.ShaderType != ShaderTypeDependencies[ShaderIndex].ShaderType
-			|| ShaderTypeDependency.SourceHash != ShaderTypeDependencies[ShaderIndex].SourceHash)
+		if (ShaderTypeDependency.ShaderType != InReferenceSet.ShaderTypeDependencies[ShaderIndex].ShaderType
+			|| ShaderTypeDependency.SourceHash != InReferenceSet.ShaderTypeDependencies[ShaderIndex].SourceHash)
 		{
 			return false;
 		}
