Fix CIS for non editor build in ocio

diff --git a/Engine/Plugins/Compositing/OpenColorIO/Source/OpenColorIO/Private/OpenColorIOShaderMap.cpp b/Engine/Plugins/Compositing/OpenColorIO/Source/OpenColorIO/Private/OpenColorIOShaderMap.cpp
--- a/Engine/Plugins/Compositing/OpenColorIO/Source/OpenColorIO/Private/OpenColorIOShaderMap.cpp
+++ b/Engine/Plugins/Compositing/OpenColorIO/Source/OpenColorIO/Private/OpenColorIOShaderMap.cpp
@@ -118,8 +118,7 @@ bool FOpenColorIOShaderMapId::operator==(const FOpenColorIOShaderMapId& InRefere
 	{
 		const FShaderTypeDependency& ShaderTypeDependency = ShaderTypeDependencies[ShaderIndex];
 
-		if (ShaderTypeDependency.ShaderType != InReferenceSet.ShaderTypeDependencies[ShaderIndex].ShaderType
-			|| ShaderTypeDependency.SourceHash != InReferenceSet.ShaderTypeDependencies[ShaderIndex].SourceHash)
+		if (ShaderTypeDependency != InReferenceSet.ShaderTypeDependencies[ShaderIndex])
 		{
 			return false;
 		}
diff --git a/Engine/Plugins/Compositing/OpenColorIO/Source/OpenColorIO/Private/OpenColorIOShared.cpp b/Engine/Plugins/Compositing/OpenColorIO/Source/OpenColorIO/Private/OpenColorIOShared.cpp
--- a/Engine/Plugins/Compositing/OpenColorIO/Source/OpenColorIO/Private/OpenColorIOShared.cpp
+++ b/Engine/Plugins/Compositing/OpenColorIO/Source/OpenColorIO/Private/OpenColorIOShared.cpp
@@ -435,6 +435,7 @@ bool FOpenColorIOTransformResource::BeginCompileShaderMap(const FOpenColorIOShad
 
 void FOpenColorIOShaderMapId::SetShaderDependencies(const TArray<FShaderType*>& InShaderTypes, EShaderPlatform InShaderPlatform)
 {
+#if WITH_EDITOR
 	if (!FPlatformProperties::RequiresCookedData())
 	{
 		for (FShaderType* ShaderType : InShaderTypes)
@@ -448,6 +449,7 @@ void FOpenColorIOShaderMapId::SetShaderDependencies(const TArray<FShaderType*>&
 			}
 		}
 	}
+#endif //WITH_EDITOR
 }
 
 bool FOpenColorIOShaderMapId::ContainsShaderType(const FShaderType* ShaderType) const
