- Fixing ocio not workin in -game
#integrating CL#7288852

diff --git a/Engine/Plugins/Compositing/OpenColorIO/Source/OpenColorIO/Private/OpenColorIOLibHandler.cpp b/Engine/Plugins/Compositing/OpenColorIO/Source/OpenColorIO/Private/OpenColorIOLibHandler.cpp
new file mode 100644
--- /dev/null
+++ b/Engine/Plugins/Compositing/OpenColorIO/Source/OpenColorIO/Private/OpenColorIOLibHandler.cpp
@@ -0,0 +1,68 @@
+// Copyright 1998-2019 Epic Games, Inc. All Rights Reserved.
+
+#include "OpenColorIOLibHandler.h"
+
+#include "HAL/PlatformProcess.h"
+#include "Interfaces/IPluginManager.h"
+#include "IOpenColorIOModule.h"
+#include "Logging/LogMacros.h"
+#include "Misc/Paths.h"
+
+
+ //~ Static initialization
+ //--------------------------------------------------------------------
+void* FOpenColorIOLibHandler::LibHandle = nullptr;
+
+
+//~ Initialization functions implementation
+//--------------------------------------------------------------------
+bool FOpenColorIOLibHandler::Initialize()
+{
+#if WITH_EDITOR
+	check(LibHandle == nullptr);
+
+	const FString PluginDir = IPluginManager::Get().FindPlugin(TEXT("OpenColorIO"))->GetBaseDir();
+	const FString OCIOBinPath = FPaths::Combine(PluginDir, TEXT(PREPROCESSOR_TO_STRING(OCIO_PLATFORM_PATH)));
+	const FString DLLPath = FPaths::Combine(OCIOBinPath, TEXT(PREPROCESSOR_TO_STRING(OCIO_DLL_NAME)));
+
+	FPlatformProcess::PushDllDirectory(*OCIOBinPath);
+
+	if (!FPaths::FileExists(DLLPath))
+	{
+		UE_LOG(LogOpenColorIO, Error, TEXT("Failed to find the OpenColorIO dll. Plug-in will not be functional."));
+		return false;
+	}
+	
+	LibHandle = FPlatformProcess::GetDllHandle(*DLLPath);
+	FPlatformProcess::PopDllDirectory(*OCIOBinPath);
+
+	if (LibHandle == nullptr)
+	{
+		UE_LOG(LogOpenColorIO, Error, TEXT("Failed to load required library %s. Plug-in will not be functional."), *DLLPath);
+		return false;
+	}
+
+#else
+	return false;
+#endif // WITH_EDITOR
+	return true;
+}
+
+bool FOpenColorIOLibHandler::IsInitialized()
+{
+	return (LibHandle != nullptr);
+}
+
+void FOpenColorIOLibHandler::Shutdown()
+{
+#if WITH_EDITOR
+	if (LibHandle != nullptr)
+	{
+		FPlatformProcess::FreeDllHandle(LibHandle);
+		LibHandle = nullptr;
+	}
+#endif // WITH_EDITOR
+}
+
+
+
diff --git a/Engine/Plugins/Compositing/OpenColorIO/Source/OpenColorIO/Private/OpenColorIOLibHandler.h b/Engine/Plugins/Compositing/OpenColorIO/Source/OpenColorIO/Private/OpenColorIOLibHandler.h
new file mode 100644
--- /dev/null
+++ b/Engine/Plugins/Compositing/OpenColorIO/Source/OpenColorIO/Private/OpenColorIOLibHandler.h
@@ -0,0 +1,15 @@
+// Copyright 1998-2019 Epic Games, Inc. All Rights Reserved.
+
+#pragma once
+
+
+class FOpenColorIOLibHandler
+{
+public:
+	static bool Initialize();
+	static bool IsInitialized();
+	static void Shutdown();
+
+private:
+	static void* LibHandle;
+};
diff --git a/Engine/Plugins/Compositing/OpenColorIO/Source/OpenColorIO/Private/OpenColorIOModule.cpp b/Engine/Plugins/Compositing/OpenColorIO/Source/OpenColorIO/Private/OpenColorIOModule.cpp
--- a/Engine/Plugins/Compositing/OpenColorIO/Source/OpenColorIO/Private/OpenColorIOModule.cpp
+++ b/Engine/Plugins/Compositing/OpenColorIO/Source/OpenColorIO/Private/OpenColorIOModule.cpp
@@ -4,9 +4,11 @@
 
 #include "Interfaces/IPluginManager.h"
 #include "Modules/ModuleManager.h"
+#include "OpenColorIOLibHandler.h"
 #include "ShaderCore.h"
 
 
+
 DEFINE_LOG_CATEGORY(LogOpenColorIO);
 
 #define LOCTEXT_NAMESPACE "OpenColorIOModule"
@@ -21,6 +23,8 @@ public:
 	//~ IModuleInterface interface
 	virtual void StartupModule() override
 	{
+		FOpenColorIOLibHandler::Initialize();
+
 		// Maps virtual shader source directory /Plugin/OpenCVLensDistortion to the plugin's actual Shaders directory.
 		FString PluginShaderDir = FPaths::Combine(IPluginManager::Get().FindPlugin(TEXT("OpenColorIO"))->GetBaseDir(), TEXT("Shaders"));
 		AddShaderSourceDirectoryMapping(TEXT("/Plugin/OpenColorIO"), PluginShaderDir);
@@ -28,6 +32,7 @@ public:
 
 	virtual void ShutdownModule() override
 	{
+		FOpenColorIOLibHandler::Shutdown();
 	}
 };
 
diff --git a/Engine/Plugins/Compositing/OpenColorIO/Source/OpenColorIOEditor/Private/OpenColorIOEditorModule.cpp b/Engine/Plugins/Compositing/OpenColorIO/Source/OpenColorIOEditor/Private/OpenColorIOEditorModule.cpp
--- a/Engine/Plugins/Compositing/OpenColorIO/Source/OpenColorIOEditor/Private/OpenColorIOEditorModule.cpp
+++ b/Engine/Plugins/Compositing/OpenColorIO/Source/OpenColorIOEditor/Private/OpenColorIOEditorModule.cpp
@@ -8,7 +8,6 @@
 #include "ISettingsModule.h"
 #include "ISettingsSection.h"
 #include "Modules/ModuleManager.h"
-#include "OpenColorIOLibHandler.h"
 #include "OpenColorIOColorSpaceConversionCustomization.h"
 #include "OpenColorIOColorSpaceCustomization.h"
 #include "OpenColorIOColorTransform.h"
@@ -29,12 +28,8 @@ class FOpenColorIOEditorModule : public IOpenColorIOEditorModule
 {
 public:
 
-	virtual bool IsInitialized() const override { return FOpenColorIOLibHandler::IsInitialized(); }
-
 	virtual void StartupModule() override
 	{
-		FOpenColorIOLibHandler::Initialize();
-
 		FWorldDelegates::OnPreWorldInitialization.AddRaw(this, &FOpenColorIOEditorModule::OnWorldInit);
 
 		// Register asset type actions for OpenColorIOConfiguration class
@@ -67,8 +62,6 @@ public:
 
 		CleanFeatureLevelDelegate();
 		FWorldDelegates::OnPreWorldInitialization.RemoveAll(this);
-
-		FOpenColorIOLibHandler::Shutdown();
 	}
 
 	void RegisterCustomizations()
diff --git a/Engine/Plugins/Compositing/OpenColorIO/Source/OpenColorIOEditor/Private/OpenColorIOLibHandler.cpp b/Engine/Plugins/Compositing/OpenColorIO/Source/OpenColorIOEditor/Private/OpenColorIOLibHandler.cpp
deleted file mode 100644
diff --git a/Engine/Plugins/Compositing/OpenColorIO/Source/OpenColorIOEditor/Private/OpenColorIOLibHandler.h b/Engine/Plugins/Compositing/OpenColorIO/Source/OpenColorIOEditor/Private/OpenColorIOLibHandler.h
deleted file mode 100644
diff --git a/Engine/Plugins/Compositing/OpenColorIO/Source/OpenColorIOEditor/Public/IOpenColorIOEditorModule.h b/Engine/Plugins/Compositing/OpenColorIO/Source/OpenColorIOEditor/Public/IOpenColorIOEditorModule.h
--- a/Engine/Plugins/Compositing/OpenColorIO/Source/OpenColorIOEditor/Public/IOpenColorIOEditorModule.h
+++ b/Engine/Plugins/Compositing/OpenColorIO/Source/OpenColorIOEditor/Public/IOpenColorIOEditorModule.h
@@ -15,8 +15,5 @@ DECLARE_LOG_CATEGORY_EXTERN(LogOpenColorIOEditor, Log, All);
 class IOpenColorIOEditorModule : public IModuleInterface
 {
 public:
-	
-	/** @return true if the OpenColorIOEditor module and OpenColorIO dll could be loaded */
-	virtual bool IsInitialized() const = 0;
 };
 
