- fixing cook for otherplatforms than windows

diff --git a/Engine/Plugins/Compositing/OpenColorIO/Source/OpenColorIO/Private/OpenColorIOLibHandler.cpp b/Engine/Plugins/Compositing/OpenColorIO/Source/OpenColorIO/Private/OpenColorIOLibHandler.cpp
--- a/Engine/Plugins/Compositing/OpenColorIO/Source/OpenColorIO/Private/OpenColorIOLibHandler.cpp
+++ b/Engine/Plugins/Compositing/OpenColorIO/Source/OpenColorIO/Private/OpenColorIOLibHandler.cpp
@@ -18,7 +18,7 @@ void* FOpenColorIOLibHandler::LibHandle = nullptr;
 //--------------------------------------------------------------------
 bool FOpenColorIOLibHandler::Initialize()
 {
-#if WITH_EDITOR
+#if WITH_EDITOR && WITH_OCIO
 	check(LibHandle == nullptr);
 
 	const FString PluginDir = IPluginManager::Get().FindPlugin(TEXT("OpenColorIO"))->GetBaseDir();
@@ -55,7 +55,7 @@ bool FOpenColorIOLibHandler::IsInitialized()
 
 void FOpenColorIOLibHandler::Shutdown()
 {
-#if WITH_EDITOR
+#if WITH_EDITOR && WITH_OCIO
 	if (LibHandle != nullptr)
 	{
 		FPlatformProcess::FreeDllHandle(LibHandle);
