Add debug context to DDC requests that take the key as a parameter

Fix new deprecation warnings after merging the API change from another stream.

#rnx

diff --git a/Engine/Plugins/Compositing/OpenColorIO/Source/OpenColorIO/Private/OpenColorIOShaderMap.cpp b/Engine/Plugins/Compositing/OpenColorIO/Source/OpenColorIO/Private/OpenColorIOShaderMap.cpp
--- a/Engine/Plugins/Compositing/OpenColorIO/Source/OpenColorIO/Private/OpenColorIOShaderMap.cpp
+++ b/Engine/Plugins/Compositing/OpenColorIO/Source/OpenColorIO/Private/OpenColorIOShaderMap.cpp
@@ -305,7 +305,7 @@ void FOpenColorIOShaderMap::LoadFromDerivedDataCache(const FOpenColorIOTransform
 			TArray<uint8> CachedData;
 			const FString DataKey = GetOpenColorIOShaderMapKeyString(InShaderMapId, InPlatform);
 
-			if (GetDerivedDataCacheRef().GetSynchronous(*DataKey, CachedData))
+			if (GetDerivedDataCacheRef().GetSynchronous(*DataKey, CachedData, InColorTransform->GetFriendlyName()))
 			{
 				COOK_STAT(Timer.AddHit(CachedData.Num()));
 				InOutShaderMap = new FOpenColorIOShaderMap();
@@ -340,7 +340,7 @@ void FOpenColorIOShaderMap::SaveToDerivedDataCache()
 	FMemoryWriter Ar(SaveData, true);
 	Serialize(Ar);
 
-	GetDerivedDataCacheRef().Put(*GetOpenColorIOShaderMapKeyString(ShaderMapId, Platform), SaveData);
+	GetDerivedDataCacheRef().Put(*GetOpenColorIOShaderMapKeyString(ShaderMapId, Platform), SaveData, GetFriendlyName());
 	COOK_STAT(Timer.AddMiss(SaveData.Num()));
 #endif
 }
diff --git a/Engine/Plugins/FX/Niagara/Source/Niagara/Private/NiagaraScript.cpp b/Engine/Plugins/FX/Niagara/Source/Niagara/Private/NiagaraScript.cpp
--- a/Engine/Plugins/FX/Niagara/Source/Niagara/Private/NiagaraScript.cpp
+++ b/Engine/Plugins/FX/Niagara/Source/Niagara/Private/NiagaraScript.cpp
@@ -1425,7 +1425,7 @@ void UNiagaraScript::RequestCompile(bool bForceCompile)
 		TSharedPtr<FNiagaraCompileRequestDataBase, ESPMode::ThreadSafe> RequestData = NiagaraModule.Precompile(this);
 
 		// check the ddc first
-		if (!GetDerivedDataCacheRef().GetSynchronous(*GetNiagaraDDCKeyString(), OutData))
+		if (!GetDerivedDataCacheRef().GetSynchronous(*GetNiagaraDDCKeyString(), OutData, GetPathName()))
 		{
 			FNiagaraVMExecutableData ExeData;
 			if (BinaryToExecData(OutData, ExeData))
@@ -1447,7 +1447,7 @@ void UNiagaraScript::RequestCompile(bool bForceCompile)
 			// save result to the ddc
 			if (ExecToBinaryData(OutData, *ExeData))
 			{
-				GetDerivedDataCacheRef().Put(*GetNiagaraDDCKeyString(), OutData);
+				GetDerivedDataCacheRef().Put(*GetNiagaraDDCKeyString(), OutData, GetPathName());
 			}
 		}
 		ActiveCompileRoots.Empty();
diff --git a/Engine/Plugins/FX/Niagara/Source/Niagara/Private/NiagaraSystem.cpp b/Engine/Plugins/FX/Niagara/Source/Niagara/Private/NiagaraSystem.cpp
--- a/Engine/Plugins/FX/Niagara/Source/Niagara/Private/NiagaraSystem.cpp
+++ b/Engine/Plugins/FX/Niagara/Source/Niagara/Private/NiagaraSystem.cpp
@@ -1002,7 +1002,7 @@ bool UNiagaraSystem::ProcessCompilationResult(FEmitterCompiledScriptPair& Script
 	TArray<uint8> OutData;
 	if (UNiagaraScript::ExecToBinaryData(OutData, *ExeData))
 	{
-		GetDerivedDataCacheRef().Put(*ScriptPair.CompiledScript->GetNiagaraDDCKeyString(), OutData);
+		GetDerivedDataCacheRef().Put(*ScriptPair.CompiledScript->GetNiagaraDDCKeyString(), OutData, GetPathName());
 		return true;
 	}
 	
@@ -1016,7 +1016,7 @@ bool UNiagaraSystem::GetFromDDC(FEmitterCompiledScriptPair& ScriptPair)
 	ScriptPair.CompileId = NewID;
 
 	TArray<uint8> Data;
-	if (GetDerivedDataCacheRef().GetSynchronous(*ScriptPair.CompiledScript->GetNiagaraDDCKeyString(), Data))
+	if (GetDerivedDataCacheRef().GetSynchronous(*ScriptPair.CompiledScript->GetNiagaraDDCKeyString(), Data, GetPathName()))
 	{
 		TSharedPtr<FNiagaraVMExecutableData> ExeData = MakeShared<FNiagaraVMExecutableData>();
 		if (ScriptPair.CompiledScript->BinaryToExecData(Data, *ExeData))
