- Removing a call to UpdateResource of ocio 3d lut in cooked games. PostLoad of the texture will take care of it.

diff --git a/Engine/Plugins/Compositing/OpenColorIO/Source/OpenColorIO/Private/OpenColorIOColorTransform.cpp b/Engine/Plugins/Compositing/OpenColorIO/Source/OpenColorIO/Private/OpenColorIOColorTransform.cpp
--- a/Engine/Plugins/Compositing/OpenColorIO/Source/OpenColorIO/Private/OpenColorIOColorTransform.cpp
+++ b/Engine/Plugins/Compositing/OpenColorIO/Source/OpenColorIO/Private/OpenColorIOColorTransform.cpp
@@ -247,11 +247,6 @@ void UOpenColorIOColorTransform::CacheResourceTextures()
 		}
 #endif
 	}
-	else
-	{
-		//This is the path for cooked data where the 3dLut is serialized in the transform asset.
-		Lut3dTexture->UpdateResource();
-	}
 }
 
 void UOpenColorIOColorTransform::CacheResourceShadersForRendering(bool bRegenerateId)
