Copying Private-LoadTimes-4.24 stream to Main. Biggest changes are in Materials/Shader memory freezing.

diff --git a/Engine/Config/Android/DataDrivenPlatformInfo.ini b/Engine/Config/Android/DataDrivenPlatformInfo.ini
--- a/Engine/Config/Android/DataDrivenPlatformInfo.ini
+++ b/Engine/Config/Android/DataDrivenPlatformInfo.ini
@@ -1,5 +1,8 @@
 [DataDrivenPlatformInfo]
 bIsConfidential=false
+AudioCompressionSettingsIniSectionName=/Script/AndroidRuntimeSettings.AndroidRuntimeSettings
+Freezing_b32Bit=ini:Engine:Android:/Script/AndroidRuntimeSettings.AndroidRuntimeSettings:bBuildForArmV7
+Freezing_bForce64BitMemoryImagePointers=true
 
 [PlatformInfo Android]
 TargetPlatformName=Android
diff --git a/Engine/Config/HoloLens/DataDrivenPlatformInfo.ini b/Engine/Config/HoloLens/DataDrivenPlatformInfo.ini
--- a/Engine/Config/HoloLens/DataDrivenPlatformInfo.ini
+++ b/Engine/Config/HoloLens/DataDrivenPlatformInfo.ini
@@ -1,5 +1,7 @@
 [DataDrivenPlatformInfo]
 bIsConfidential=false
+Freezing_MaxFieldAlignment=8
+Freezing_bAlignBases=true
 
 [PlatformInfo HoloLens]
 TargetPlatformName=HoloLens
diff --git a/Engine/Config/IOS/DataDrivenPlatformInfo.ini b/Engine/Config/IOS/DataDrivenPlatformInfo.ini
--- a/Engine/Config/IOS/DataDrivenPlatformInfo.ini
+++ b/Engine/Config/IOS/DataDrivenPlatformInfo.ini
@@ -1,5 +1,6 @@
 [DataDrivenPlatformInfo]
 bIsConfidential=false
+AudioCompressionSettingsIniSectionName=/Script/IOSRuntimeSettings.IOSRuntimeSettings
 
 [PlatformInfo IOS]
 TargetPlatformName=IOS
diff --git a/Engine/Config/TVOS/DataDrivenPlatformInfo.ini b/Engine/Config/TVOS/DataDrivenPlatformInfo.ini
--- a/Engine/Config/TVOS/DataDrivenPlatformInfo.ini
+++ b/Engine/Config/TVOS/DataDrivenPlatformInfo.ini
@@ -1,6 +1,7 @@
 [DataDrivenPlatformInfo]
 bIsConfidential=false
 IniParent=IOS
+AudioCompressionSettingsIniSectionName=/Script/IOSRuntimeSettings.IOSRuntimeSettings
 
 [PlatformInfo TVOS]
 TargetPlatformName=TVOS
diff --git a/Engine/Config/Windows/DataDrivenPlatformInfo.ini b/Engine/Config/Windows/DataDrivenPlatformInfo.ini
--- a/Engine/Config/Windows/DataDrivenPlatformInfo.ini
+++ b/Engine/Config/Windows/DataDrivenPlatformInfo.ini
@@ -1,5 +1,11 @@
 [DataDrivenPlatformInfo]
 bIsConfidential=false
+AudioCompressionSettingsIniSectionName=/Script/WindowsTargetPlatform.WindowsTargetSettings
+Freezing_b32Bit=ini:Engine:Windows:/Script/WindowsTargetPlatform.WindowsTargetSettings:bTarget32Bit
+Freezing_bWithRayTracing=!ini:Engine:Windows:/Script/WindowsTargetPlatform.WindowsTargetSettings:bTarget32Bit
+Freezing_MaxFieldAlignment32=8
+Freezing_MaxFieldAlignment64=16
+Freezing_bAlignBases=true
 
 [PlatformInfo WindowsNoEditor]
 TargetPlatformName=WindowsNoEditor
@@ -130,4 +136,5 @@ PlatformGroupName=Desktop
 [ShaderPlatform PCD3D_SM5]
 bSupportsRayTracing=true
 bSupportsRayTracingMissShaderBindings=true
+bSupportsGPUScene=true
 
diff --git a/Engine/Platforms/XXX/Source/Developer/XXXTargetPlatform/Private/XXXTargetPlatform.h b/Engine/Platforms/XXX/Source/Developer/XXXTargetPlatform/Private/XXXTargetPlatform.h
--- a/Engine/Platforms/XXX/Source/Developer/XXXTargetPlatform/Private/XXXTargetPlatform.h
+++ b/Engine/Platforms/XXX/Source/Developer/XXXTargetPlatform/Private/XXXTargetPlatform.h
@@ -80,11 +80,6 @@ public:
 	// this should be all the possible return values from the function GetWaveFormat
 	virtual void GetAllWaveFormats( TArray<FName>& OutFormats ) const override;
 
-	virtual FPlatformAudioCookOverrides* GetAudioCompressionSettings() const override
-	{
-		return nullptr;
-	}
-
 #endif // WITH_ENGINE
 
 	DECLARE_DERIVED_EVENT(FXXXTargetPlatform, ITargetPlatform::FOnTargetDeviceDiscovered, FOnTargetDeviceDiscovered);
diff --git a/Engine/Plugins/2D/Paper2D/Source/Paper2D/Private/PaperRenderSceneProxy.cpp b/Engine/Plugins/2D/Paper2D/Source/Paper2D/Private/PaperRenderSceneProxy.cpp
--- a/Engine/Plugins/2D/Paper2D/Source/Paper2D/Private/PaperRenderSceneProxy.cpp
+++ b/Engine/Plugins/2D/Paper2D/Source/Paper2D/Private/PaperRenderSceneProxy.cpp
@@ -101,37 +101,39 @@ public:
 		return Parent->GetMaterialWithFallback(InFeatureLevel, OutFallbackMaterialRenderProxy);
 	}
 
-	virtual bool GetVectorValue(const FMaterialParameterInfo& ParameterInfo, FLinearColor* OutValue, const FMaterialRenderContext& Context) const override
+	virtual bool GetVectorValue(const FHashedMaterialParameterInfo& ParameterInfo, FLinearColor* OutValue, const FMaterialRenderContext& Context) const override
 	{
 		return Parent->GetVectorValue(ParameterInfo, OutValue, Context);
 	}
 
-	virtual bool GetScalarValue(const FMaterialParameterInfo& ParameterInfo, float* OutValue, const FMaterialRenderContext& Context) const override
+	virtual bool GetScalarValue(const FHashedMaterialParameterInfo& ParameterInfo, float* OutValue, const FMaterialRenderContext& Context) const override
 	{
 		return Parent->GetScalarValue(ParameterInfo, OutValue, Context);
 	}
 
-	virtual bool GetTextureValue(const FMaterialParameterInfo& ParameterInfo, const UTexture** OutValue, const FMaterialRenderContext& Context) const override
+	virtual bool GetTextureValue(const FHashedMaterialParameterInfo& ParameterInfo, const UTexture** OutValue, const FMaterialRenderContext& Context) const override
 	{
 		if (ParameterInfo.Name == TextureParameterName)
 		{
 			*OutValue = ApplyEditorOverrides(BaseTexture);
 			return true;
 		}
-		else if (ParameterInfo.Name.GetComparisonIndex() == AdditionalTextureParameterRootName.GetComparisonIndex())
+
+		for (int32 AdditionalSlotIndex = 0; AdditionalSlotIndex < AdditionalTextures.Num(); ++AdditionalSlotIndex)
 		{
-			const int32 AdditionalSlotIndex = ParameterInfo.Name.GetNumber() - 1;
-			if (AdditionalTextures.IsValidIndex(AdditionalSlotIndex))
+			FName AdditionalSlotName = AdditionalTextureParameterRootName;
+			AdditionalSlotName.SetNumber(AdditionalSlotIndex + 1);
+			if (ParameterInfo.Name == AdditionalSlotName)
 			{
 				*OutValue = ApplyEditorOverrides(AdditionalTextures[AdditionalSlotIndex]);
 				return true;
 			}
 		}
-		
+
 		return Parent->GetTextureValue(ParameterInfo, OutValue, Context);
 	}
 
-	virtual bool GetTextureValue(const FMaterialParameterInfo& ParameterInfo, const URuntimeVirtualTexture** OutValue, const FMaterialRenderContext& Context) const override
+	virtual bool GetTextureValue(const FHashedMaterialParameterInfo& ParameterInfo, const URuntimeVirtualTexture** OutValue, const FMaterialRenderContext& Context) const override
 	{
 		return Parent->GetTextureValue(ParameterInfo, OutValue, Context);
 	}
diff --git a/Engine/Plugins/Compositing/LensDistortion/Source/LensDistortion/Private/LensDistortionRendering.cpp b/Engine/Plugins/Compositing/LensDistortion/Source/LensDistortion/Private/LensDistortionRendering.cpp
--- a/Engine/Plugins/Compositing/LensDistortion/Source/LensDistortion/Private/LensDistortionRendering.cpp
+++ b/Engine/Plugins/Compositing/LensDistortion/Source/LensDistortion/Private/LensDistortionRendering.cpp
@@ -54,7 +54,9 @@ static FVector2D LensUndistortViewportUVIntoViewSpace(
 
 class FLensDistortionUVGenerationShader : public FGlobalShader
 {
+	DECLARE_INLINE_TYPE_LAYOUT(FLensDistortionUVGenerationShader, NonVirtual);
 public:
+
 	static bool ShouldCompilePermutation(const FGlobalShaderPermutationParameters& Parameters)
 	{
 		return IsFeatureLevelSupported(Parameters.Platform, ERHIFeatureLevel::SM5);
@@ -105,28 +107,20 @@ public:
 		SetShaderValue(RHICmdList, ShaderRHI, OutputMultiplyAndAdd, CompiledCameraModel.OutputMultiplyAndAdd);
 	}
 
-	virtual bool Serialize(FArchive& Ar) override
-	{
-		bool bShaderHasOutdatedParameters = FGlobalShader::Serialize(Ar);
-		Ar << PixelUVSize << RadialDistortionCoefs << TangentialDistortionCoefs << DistortedCameraMatrix << UndistortedCameraMatrix << OutputMultiplyAndAdd;
-		return bShaderHasOutdatedParameters;
-	}
-
 private:
-	FShaderParameter PixelUVSize;
-	FShaderParameter RadialDistortionCoefs;
-	FShaderParameter TangentialDistortionCoefs;
-	FShaderParameter DistortedCameraMatrix;
-	FShaderParameter UndistortedCameraMatrix;
-	FShaderParameter OutputMultiplyAndAdd;
-
+	
+	LAYOUT_FIELD(FShaderParameter, PixelUVSize);
+	LAYOUT_FIELD(FShaderParameter, RadialDistortionCoefs);
+	LAYOUT_FIELD(FShaderParameter, TangentialDistortionCoefs);
+	LAYOUT_FIELD(FShaderParameter, DistortedCameraMatrix);
+	LAYOUT_FIELD(FShaderParameter, UndistortedCameraMatrix);
+	LAYOUT_FIELD(FShaderParameter, OutputMultiplyAndAdd);
 };
 
 
 class FLensDistortionUVGenerationVS : public FLensDistortionUVGenerationShader
 {
 	DECLARE_SHADER_TYPE(FLensDistortionUVGenerationVS, Global);
-
 public:
 
 	/** Default constructor. */
@@ -143,7 +137,6 @@ public:
 class FLensDistortionUVGenerationPS : public FLensDistortionUVGenerationShader
 {
 	DECLARE_SHADER_TYPE(FLensDistortionUVGenerationPS, Global);
-
 public:
 
 	/** Default constructor. */
@@ -192,7 +185,7 @@ static void DrawUVDisplacementToRenderTarget_RenderThread(
 			DisplacementMapResolution.X, DisplacementMapResolution.Y, 1.f);
 
 		// Get shaders.
-		TShaderMap<FGlobalShaderType>* GlobalShaderMap = GetGlobalShaderMap(FeatureLevel);
+		FGlobalShaderMap* GlobalShaderMap = GetGlobalShaderMap(FeatureLevel);
 		TShaderMapRef< FLensDistortionUVGenerationVS > VertexShader(GlobalShaderMap);
 		TShaderMapRef< FLensDistortionUVGenerationPS > PixelShader(GlobalShaderMap);
 
@@ -204,8 +197,8 @@ static void DrawUVDisplacementToRenderTarget_RenderThread(
 		GraphicsPSOInit.RasterizerState = TStaticRasterizerState<>::GetRHI();
 		GraphicsPSOInit.PrimitiveType = PT_TriangleList;
 		GraphicsPSOInit.BoundShaderState.VertexDeclarationRHI = GetVertexDeclarationFVector4();
-		GraphicsPSOInit.BoundShaderState.VertexShaderRHI = GETSAFERHISHADER_VERTEX(*VertexShader);
-		GraphicsPSOInit.BoundShaderState.PixelShaderRHI = GETSAFERHISHADER_PIXEL(*PixelShader);
+		GraphicsPSOInit.BoundShaderState.VertexShaderRHI = VertexShader.GetVertexShader();
+		GraphicsPSOInit.BoundShaderState.PixelShaderRHI = PixelShader.GetPixelShader();
 		SetGraphicsPipelineState(RHICmdList, GraphicsPSOInit);
 
 		// Update viewport.
@@ -214,8 +207,8 @@ static void DrawUVDisplacementToRenderTarget_RenderThread(
 			OutTextureRenderTargetResource->GetSizeX(), OutTextureRenderTargetResource->GetSizeY(), 1.f);
 
 		// Update shader uniform parameters.
-		VertexShader->SetParameters(RHICmdList, VertexShader->GetVertexShader(), CompiledCameraModel, DisplacementMapResolution);
-		PixelShader->SetParameters(RHICmdList, PixelShader->GetPixelShader(), CompiledCameraModel, DisplacementMapResolution);
+		VertexShader->SetParameters(RHICmdList, VertexShader.GetVertexShader(), CompiledCameraModel, DisplacementMapResolution);
+		PixelShader->SetParameters(RHICmdList, PixelShader.GetPixelShader(), CompiledCameraModel, DisplacementMapResolution);
 
 		// Draw grid.
 		uint32 PrimitiveCount = kGridSubdivisionX * kGridSubdivisionY * 2;
diff --git a/Engine/Plugins/Compositing/OpenCVLensDistortion/Source/OpenCVLensDistortion/Private/LensDistortionDisplacementMapRendering.cpp b/Engine/Plugins/Compositing/OpenCVLensDistortion/Source/OpenCVLensDistortion/Private/LensDistortionDisplacementMapRendering.cpp
--- a/Engine/Plugins/Compositing/OpenCVLensDistortion/Source/OpenCVLensDistortion/Private/LensDistortionDisplacementMapRendering.cpp
+++ b/Engine/Plugins/Compositing/OpenCVLensDistortion/Source/OpenCVLensDistortion/Private/LensDistortionDisplacementMapRendering.cpp
@@ -22,6 +22,7 @@ static const uint32 kGridSubdivisionY = 16;
 #if WITH_OPENCV
 class FLensDistortionDisplacementMapGenerationShader : public FGlobalShader
 {
+	DECLARE_INLINE_TYPE_LAYOUT(FLensDistortionDisplacementMapGenerationShader, NonVirtual);
 public:
 	static bool ShouldCompilePermutation(const FGlobalShaderPermutationParameters& Parameters)
 	{
@@ -54,17 +55,11 @@ public:
 		SetTextureParameter(RHICmdList, ShaderRHI, UndistortDisplacementMap, BilinearSampler, TStaticSamplerState<SF_Bilinear, AM_Clamp, AM_Clamp, AM_Clamp>::GetRHI(), PreComputedDisplacementMap->TextureRHI);
 	}
 
-	virtual bool Serialize(FArchive& Ar) override
-	{
-		bool bShaderHasOutdatedParameters = FGlobalShader::Serialize(Ar);
-		Ar << PixelUVSize << UndistortDisplacementMap << BilinearSampler;
-		return bShaderHasOutdatedParameters;
-	}
-
 private:
-	FShaderParameter PixelUVSize;
-	FShaderResourceParameter UndistortDisplacementMap;
-	FShaderResourceParameter BilinearSampler;
+	
+	LAYOUT_FIELD(FShaderParameter, PixelUVSize);
+	LAYOUT_FIELD(FShaderResourceParameter, UndistortDisplacementMap);
+	LAYOUT_FIELD(FShaderResourceParameter, BilinearSampler);
 };
 
 
@@ -133,7 +128,7 @@ static void DrawUVDisplacementToRenderTargetFromPreComputedDisplacementMap_Rende
 		RHICmdList.SetViewport(0, 0, 0.f, DisplacementMapResolution.X, DisplacementMapResolution.Y, 1.f);
 
 		// Get shaders.
-		TShaderMap<FGlobalShaderType>* GlobalShaderMap = GetGlobalShaderMap(FeatureLevel);
+		FGlobalShaderMap* GlobalShaderMap = GetGlobalShaderMap(FeatureLevel);
 		TShaderMapRef< FLensDistortionDisplacementMapGenerationVS > VertexShader(GlobalShaderMap);
 		TShaderMapRef< FLensDistortionDisplacementMapGenerationPS > PixelShader(GlobalShaderMap);
 
@@ -145,13 +140,13 @@ static void DrawUVDisplacementToRenderTargetFromPreComputedDisplacementMap_Rende
 		GraphicsPSOInit.RasterizerState = TStaticRasterizerState<>::GetRHI();
 		GraphicsPSOInit.PrimitiveType = PT_TriangleList;
 		GraphicsPSOInit.BoundShaderState.VertexDeclarationRHI = GetVertexDeclarationFVector4();
-		GraphicsPSOInit.BoundShaderState.VertexShaderRHI = GETSAFERHISHADER_VERTEX(*VertexShader);
-		GraphicsPSOInit.BoundShaderState.PixelShaderRHI = GETSAFERHISHADER_PIXEL(*PixelShader);
+		GraphicsPSOInit.BoundShaderState.VertexShaderRHI = VertexShader.GetVertexShader();
+		GraphicsPSOInit.BoundShaderState.PixelShaderRHI = PixelShader.GetPixelShader();
 		SetGraphicsPipelineState(RHICmdList, GraphicsPSOInit);
 
 		// Update shader uniform parameters.
-		VertexShader->SetParameters(RHICmdList, VertexShader->GetVertexShader(), PreComputedDisplacementMap, DisplacementMapResolution);
-		PixelShader->SetParameters(RHICmdList, PixelShader->GetPixelShader(), PreComputedDisplacementMap, DisplacementMapResolution);
+		VertexShader->SetParameters(RHICmdList, VertexShader.GetVertexShader(), PreComputedDisplacementMap, DisplacementMapResolution);
+		PixelShader->SetParameters(RHICmdList, PixelShader.GetPixelShader(), PreComputedDisplacementMap, DisplacementMapResolution);
 
 		// Draw grid.
 		const uint32 PrimitiveCount = kGridSubdivisionX * kGridSubdivisionY * 2;
diff --git a/Engine/Plugins/Compositing/OpenColorIO/Source/OpenColorIO/Private/OpenColorIOColorTransform.cpp b/Engine/Plugins/Compositing/OpenColorIO/Source/OpenColorIO/Private/OpenColorIOColorTransform.cpp
--- a/Engine/Plugins/Compositing/OpenColorIO/Source/OpenColorIO/Private/OpenColorIOColorTransform.cpp
+++ b/Engine/Plugins/Compositing/OpenColorIO/Source/OpenColorIO/Private/OpenColorIOColorTransform.cpp
@@ -77,11 +77,6 @@ void UOpenColorIOColorTransform::ProcessSerializedShaderMaps(UOpenColorIOColorTr
 {
 	check(IsInGameThread());
 
-	for (FOpenColorIOTransformResource& Resource : LoadedResources)
-	{
-		Resource.RegisterShaderMap();
-	}
-
 	for (int32 ResourceIndex = 0; ResourceIndex < LoadedResources.Num(); ResourceIndex++)
 	{
 		FOpenColorIOTransformResource& LoadedResource = LoadedResources[ResourceIndex];
diff --git a/Engine/Plugins/Compositing/OpenColorIO/Source/OpenColorIO/Private/OpenColorIODerivedDataVersion.h b/Engine/Plugins/Compositing/OpenColorIO/Source/OpenColorIO/Private/OpenColorIODerivedDataVersion.h
--- a/Engine/Plugins/Compositing/OpenColorIO/Source/OpenColorIO/Private/OpenColorIODerivedDataVersion.h
+++ b/Engine/Plugins/Compositing/OpenColorIO/Source/OpenColorIO/Private/OpenColorIODerivedDataVersion.h
@@ -10,4 +10,4 @@ OpenColorIOShaderDerivedDataVersion.h: Derived data version for OpenColorIO shad
 
 // In case of merge conflicts with DDC versions, you MUST generate a new GUID and set this new
 // guid as version
-#define OPENCOLORIO_DERIVEDDATA_VER		TEXT("A0CDC343C4844554AA76FCA6A70C594E")
\ No newline at end of file
+#define OPENCOLORIO_DERIVEDDATA_VER		TEXT("883BB66A2A0E4A5BB85621F97919BAF8")
\ No newline at end of file
diff --git a/Engine/Plugins/Compositing/OpenColorIO/Source/OpenColorIO/Private/OpenColorIORendering.cpp b/Engine/Plugins/Compositing/OpenColorIO/Source/OpenColorIO/Private/OpenColorIORendering.cpp
--- a/Engine/Plugins/Compositing/OpenColorIO/Source/OpenColorIO/Private/OpenColorIORendering.cpp
+++ b/Engine/Plugins/Compositing/OpenColorIO/Source/OpenColorIO/Private/OpenColorIORendering.cpp
@@ -45,14 +45,6 @@ public:
 		: FGlobalShader(Initializer)
 	{
 	}
-
-public:
-
-	virtual bool Serialize(FArchive& Ar) override
-	{
-		bool bShaderHasOutdatedParameters = FGlobalShader::Serialize(Ar);
-		return bShaderHasOutdatedParameters;
-	}
 };
 
 
@@ -82,9 +74,9 @@ static void ProcessOCIOColorSpaceTransform_RenderThread(
 	InRHICmdList.SetViewport(0, 0, 0.f, Resolution.X, Resolution.Y, 1.f);
 
 	// Get shader from shader map.
-	TShaderMap<FGlobalShaderType>* GlobalShaderMap = GetGlobalShaderMap(InFeatureLevel);
+	FGlobalShaderMap* GlobalShaderMap = GetGlobalShaderMap(InFeatureLevel);
 	TShaderMapRef<FOpenColorIOVertexShader> VertexShader(GlobalShaderMap);
-	FOpenColorIOPixelShader* OCIOPixelShader = InOCIOColorTransformResource->GetShader();
+	TShaderRef<FOpenColorIOPixelShader> OCIOPixelShader = InOCIOColorTransformResource->GetShader();
 
 	// Set the graphic pipeline state.
 	FGraphicsPipelineStateInitializer GraphicsPSOInit;
@@ -94,8 +86,8 @@ static void ProcessOCIOColorSpaceTransform_RenderThread(
 	GraphicsPSOInit.RasterizerState = TStaticRasterizerState<>::GetRHI();
 	GraphicsPSOInit.PrimitiveType = PT_TriangleList;
 	GraphicsPSOInit.BoundShaderState.VertexDeclarationRHI = GetVertexDeclarationFVector4();
-	GraphicsPSOInit.BoundShaderState.VertexShaderRHI = GETSAFERHISHADER_VERTEX(*VertexShader);
-	GraphicsPSOInit.BoundShaderState.PixelShaderRHI = GETSAFERHISHADER_PIXEL(OCIOPixelShader);
+	GraphicsPSOInit.BoundShaderState.VertexShaderRHI = VertexShader.GetVertexShader();
+	GraphicsPSOInit.BoundShaderState.PixelShaderRHI = OCIOPixelShader.GetPixelShader();
 	SetGraphicsPipelineState(InRHICmdList, GraphicsPSOInit);
 
 	// Update pixel shader parameters
@@ -165,7 +157,7 @@ bool FOpenColorIORendering::ApplyColorTransform(UWorld* InWorld, const FOpenColo
 
 	check(ShaderResource);
 
-	if (ShaderResource->GetShaderGameThread() == nullptr)
+	if (ShaderResource->GetShaderGameThread().IsNull())
 	{
 		UE_LOG(LogOpenColorIO, Warning, TEXT("OCIOPass - Shader was invalid for Resource %s"), *ShaderResource->GetFriendlyName());
 		return false;
diff --git a/Engine/Plugins/Compositing/OpenColorIO/Source/OpenColorIO/Private/OpenColorIOShader.cpp b/Engine/Plugins/Compositing/OpenColorIO/Source/OpenColorIO/Private/OpenColorIOShader.cpp
--- a/Engine/Plugins/Compositing/OpenColorIO/Source/OpenColorIO/Private/OpenColorIOShader.cpp
+++ b/Engine/Plugins/Compositing/OpenColorIO/Source/OpenColorIO/Private/OpenColorIOShader.cpp
@@ -15,6 +15,7 @@
 #include "Stats/StatsMisc.h"
 #include "UObject/UObjectGlobals.h"
 
+IMPLEMENT_TYPE_LAYOUT(FOpenColorIOShader);
 
 IMPLEMENT_SHADER_TYPE(, FOpenColorIOPixelShader, TEXT("/Plugin/OpenColorIO/Private/OpenColorIOShader.usf"), TEXT("MainPS"), SF_Pixel)
 
@@ -23,7 +24,7 @@ IMPLEMENT_SHADER_TYPE(, FOpenColorIOPixelShader, TEXT("/Plugin/OpenColorIO/Priva
 
 
 FOpenColorIOPixelShader::FOpenColorIOPixelShader(const FOpenColorIOShaderType::CompiledShaderInitializerType& Initializer)
-	: FShader(Initializer)
+	: FOpenColorIOShader(Initializer)
 	, DebugDescription(Initializer.DebugDescription)
 {
 	BindParams(Initializer.ParameterMap);
@@ -31,12 +32,12 @@ FOpenColorIOPixelShader::FOpenColorIOPixelShader(const FOpenColorIOShaderType::C
 
 void FOpenColorIOPixelShader::SetParameters(FRHICommandList& InRHICmdList, FTextureResource* InInputTexture)
 {
-	SetTextureParameter(InRHICmdList, GetPixelShader(), InputTexture, InputTextureSampler, TStaticSamplerState<SF_Bilinear, AM_Clamp, AM_Clamp, AM_Clamp>::GetRHI(), InInputTexture->TextureRHI);
+	SetTextureParameter(InRHICmdList, InRHICmdList.GetBoundPixelShader(), InputTexture, InputTextureSampler, TStaticSamplerState<SF_Bilinear, AM_Clamp, AM_Clamp, AM_Clamp>::GetRHI(), InInputTexture->TextureRHI);
 }
 
 void FOpenColorIOPixelShader::SetLUTParameter(FRHICommandList& InRHICmdList, FTextureResource* InLUT3dResource)
 {
-	SetTextureParameter(InRHICmdList, GetPixelShader(), OCIO3dTexture, OCIO3dTextureSampler, TStaticSamplerState<SF_Bilinear, AM_Clamp, AM_Clamp, AM_Clamp>::GetRHI(), InLUT3dResource->TextureRHI);
+	SetTextureParameter(InRHICmdList, InRHICmdList.GetBoundPixelShader(), OCIO3dTexture, OCIO3dTextureSampler, TStaticSamplerState<SF_Bilinear, AM_Clamp, AM_Clamp, AM_Clamp>::GetRHI(), InLUT3dResource->TextureRHI);
 }
 
 void FOpenColorIOPixelShader::BindParams(const FShaderParameterMap &ParameterMap)
@@ -48,23 +49,3 @@ void FOpenColorIOPixelShader::BindParams(const FShaderParameterMap &ParameterMap
 	OCIO3dTextureSampler.Bind(ParameterMap, TEXT("ociolut3d_0Sampler"));
 }
 
-bool FOpenColorIOPixelShader::Serialize(FArchive& Ar)
-{
-	const bool bShaderHasOutdatedParameters = FShader::Serialize(Ar);
-
-	Ar << InputTexture;	
-	Ar << InputTextureSampler;
-
-	Ar << OCIO3dTexture;
-	Ar << OCIO3dTextureSampler;
-
-	Ar << DebugDescription;
-
-	return bShaderHasOutdatedParameters;
-}
-
-uint32 FOpenColorIOPixelShader::GetAllocatedSize() const
-{
-	return FShader::GetAllocatedSize() + DebugDescription.GetAllocatedSize();
-}
-
diff --git a/Engine/Plugins/Compositing/OpenColorIO/Source/OpenColorIO/Private/OpenColorIOShaderMap.cpp b/Engine/Plugins/Compositing/OpenColorIO/Source/OpenColorIO/Private/OpenColorIOShaderMap.cpp
--- a/Engine/Plugins/Compositing/OpenColorIO/Source/OpenColorIO/Private/OpenColorIOShaderMap.cpp
+++ b/Engine/Plugins/Compositing/OpenColorIO/Source/OpenColorIO/Private/OpenColorIOShaderMap.cpp
@@ -73,15 +73,14 @@ void UpdateOpenColorIOShaderCompilingStats(const FOpenColorIOTransformResource*
 	INC_DWORD_STAT_BY(STAT_ShaderCompiling_NumTotalOpenColorIOShaders,1);
 }
 
-
-void FOpenColorIOShaderMapId::Serialize(FArchive& Ar)
+/*void FOpenColorIOShaderMapId::Serialize(FArchive& Ar)
 {
 	// You must bump OPENCOLORIO_DERIVEDDATA_VER if changing the serialization of FOpenColorIOShaderMapId.
 
 	Ar << ShaderCodeHash;
 	Ar << (int32&)FeatureLevel;
 	Ar << ShaderTypeDependencies;
-}
+}*/
 
 /** Hashes the color transform specific part of this shader map Id. */
 void FOpenColorIOShaderMapId::GetOpenColorIOHash(FSHAHash& OutHash) const
@@ -130,23 +129,36 @@ bool FOpenColorIOShaderMapId::operator==(const FOpenColorIOShaderMapId& InRefere
 void FOpenColorIOShaderMapId::AppendKeyString(FString& OutKeyString) const
 {
 #if WITH_EDITOR
+	FPlatformTypeLayoutParameters LayoutParams;
+	LayoutParams.InitializeForCurrent();
+
 	OutKeyString += ShaderCodeHash;
 	OutKeyString += TEXT("_");
 
 	FString FeatureLevelString;
 	GetFeatureLevelName(FeatureLevel, FeatureLevelString);
 
+	{
+		const FSHAHash LayoutHash = Freeze::HashLayout(StaticGetTypeLayoutDesc<FOpenColorIOShaderMapContent>(), LayoutParams);
+		OutKeyString += TEXT("_");
+		OutKeyString += LayoutHash.ToString();
+		OutKeyString += TEXT("_");
+	}
+
 	TMap<const TCHAR*, FCachedUniformBufferDeclaration> ReferencedUniformBuffers;
 
 	// Add the inputs for any shaders that are stored inline in the shader map
 	for (const FShaderTypeDependency& ShaderTypeDependency : ShaderTypeDependencies)
 	{
+		const FShaderType* ShaderType = FindShaderTypeByName(ShaderTypeDependency.ShaderTypeName);
 		OutKeyString += TEXT("_");
-		OutKeyString += ShaderTypeDependency.ShaderType->GetName();
+		OutKeyString += ShaderType->GetName();
 		OutKeyString += ShaderTypeDependency.SourceHash.ToString();
-		ShaderTypeDependency.ShaderType->GetSerializationHistory().AppendKeyString(OutKeyString);
+		
+		const FSHAHash LayoutHash = Freeze::HashLayout(ShaderType->GetLayout(), LayoutParams);
+		OutKeyString += LayoutHash.ToString();
 
-		const TMap<const TCHAR*, FCachedUniformBufferDeclaration>& ReferencedUniformBufferStructsCache = ShaderTypeDependency.ShaderType->GetReferencedUniformBufferStructsCache();
+		const TMap<const TCHAR*, FCachedUniformBufferDeclaration>& ReferencedUniformBufferStructsCache = ShaderType->GetReferencedUniformBufferStructsCache();
 
 		for (TMap<const TCHAR*, FCachedUniformBufferDeclaration>::TConstIterator It(ReferencedUniformBufferStructsCache); It; ++It)
 		{
@@ -225,25 +237,19 @@ FShaderCompileJob* FOpenColorIOShaderType::BeginCompileShader(
 FShader* FOpenColorIOShaderType::FinishCompileShader(
 	const FSHAHash& InShaderMapHash,
 	const FShaderCompileJob& InCurrentJob,
-	const FString& InDebugDescription
+	const FString& InDebugDescription,
+	FShaderMapResourceBuilder& ResourceBuilder
 	)
 {
 	check(InCurrentJob.bSucceeded);
 
 	// Reuse an existing resource with the same key or create a new one based on the compile output
 	// This allows FShaders to share compiled bytecode and RHI shader references
-	TRefCountPtr<FShaderResource> Resource = FShaderResource::FindOrCreate(InCurrentJob.Output, 0);
+	const int32 ResourceIndex = ResourceBuilder.FindOrAddCode(InCurrentJob.Output);
 
-	// Find a shader with the same key in memory
-	FShader* Shader = InCurrentJob.ShaderType->FindShaderByKey(FShaderKey(InShaderMapHash, nullptr, nullptr, /* SpecificPermutationId = */ 0, InCurrentJob.Input.Target.GetPlatform()));
-
-	// There was no shader with the same key so create a new one with the compile output, which will bind shader parameters
-	if (!Shader)
-	{
-		const int32 PermutationId = 0;
-		Shader = (*ConstructCompiledRef)(FOpenColorIOShaderType::CompiledShaderInitializerType(this, PermutationId, InCurrentJob.Output, Resource, InShaderMapHash, InDebugDescription));
-		InCurrentJob.Output.ParameterMap.VerifyBindingsAreComplete(GetName(), InCurrentJob.Output.Target, InCurrentJob.VFType);
-	}
+	const int32 PermutationId = 0;
+	FShader* Shader = ConstructCompiled(FOpenColorIOShaderType::CompiledShaderInitializerType(this, PermutationId, InCurrentJob.Output, ResourceIndex, InShaderMapHash, InDebugDescription));
+	InCurrentJob.Output.ParameterMap.VerifyBindingsAreComplete(GetName(), InCurrentJob.Output.Target, InCurrentJob.VFType);
 
 	return Shader;
 }
@@ -290,7 +296,7 @@ void FOpenColorIOShaderMap::LoadFromDerivedDataCache(const FOpenColorIOTransform
 #if WITH_EDITOR
 	if (InOutShaderMap != nullptr)
 	{
-		check(InOutShaderMap->Platform == InPlatform);
+		check(InOutShaderMap->GetShaderPlatform() == InPlatform);
 		// If the shader map was non-NULL then it was found in memory but is incomplete, attempt to load the missing entries from memory
 		InOutShaderMap->LoadMissingShadersFromMemory(InColorTransform);
 	}
@@ -313,7 +319,7 @@ void FOpenColorIOShaderMap::LoadFromDerivedDataCache(const FOpenColorIOTransform
 
 				// Deserialize from the cached data
 				InOutShaderMap->Serialize(Ar);
-				InOutShaderMap->RegisterSerializedShaders(false);
+				//InOutShaderMap->RegisterSerializedShaders(false);
 
 				checkSlow(InOutShaderMap->GetShaderMapId() == InShaderMapId);
 
@@ -340,7 +346,7 @@ void FOpenColorIOShaderMap::SaveToDerivedDataCache()
 	FMemoryWriter Ar(SaveData, true);
 	Serialize(Ar);
 
-	GetDerivedDataCacheRef().Put(*GetOpenColorIOShaderMapKeyString(ShaderMapId, Platform), SaveData, GetFriendlyName());
+	GetDerivedDataCacheRef().Put(*GetOpenColorIOShaderMapKeyString(GetContent()->ShaderMapId, GetShaderPlatform()), SaveData, FStringView(*GetFriendlyName()));
 	COOK_STAT(Timer.AddMiss(SaveData.Num()));
 #endif
 }
@@ -397,10 +403,11 @@ void FOpenColorIOShaderMap::Compile(FOpenColorIOTransformResource* InColorTransf
 			InColorTransform->SetupShaderCompilationEnvironment(InPlatform, *InCompilationEnvironment);
   
 			// Store the ColorTransform name for debugging purposes.
-			FriendlyName = InColorTransform->GetFriendlyName();
-			OpenColorIOCompilationOutput = InOpenColorIOCompilationOutput;
-			ShaderMapId = InShaderMapId;
-			Platform = InPlatform;
+			FOpenColorIOShaderMapContent* NewContent = new FOpenColorIOShaderMapContent(InPlatform);
+			NewContent->FriendlyName = InColorTransform->GetFriendlyName();
+			NewContent->OpenColorIOCompilationOutput = InOpenColorIOCompilationOutput;
+			NewContent->ShaderMapId = InShaderMapId;
+			AssignContent(NewContent);
 
 			uint32 NumShaders = 0;
 			TArray<TSharedRef<FShaderCommonCompileJob, ESPMode::ThreadSafe>> NewJobs;
@@ -419,7 +426,7 @@ void FOpenColorIOShaderMap::Compile(FOpenColorIOTransformResource* InColorTransf
 					TArray<FString> ShaderErrors;
   
 					// Only compile the shader if we don't already have it
-					if (!HasShader(ShaderType, /* PermutationId = */ 0))
+					if (!NewContent->HasShader(ShaderType, /* PermutationId = */ 0))
 					{
 						auto* Job = ShaderType->BeginCompileShader(
 							CompilingId,
@@ -427,7 +434,7 @@ void FOpenColorIOShaderMap::Compile(FOpenColorIOTransformResource* InColorTransf
 							InCompilationEnvironment,
 							InPlatform,
 							NewJobs,
-							FShaderTarget(ShaderType->GetFrequency(), Platform)
+							FShaderTarget(ShaderType->GetFrequency(), GetShaderPlatform())
 							);
 						check(!SharedShaderJobs.Find(ShaderType));
 						SharedShaderJobs.Add(ShaderType, Job);
@@ -463,13 +470,13 @@ void FOpenColorIOShaderMap::Compile(FOpenColorIOTransformResource* InColorTransf
 			{
 				TArray<int32> CurrentShaderMapId;
 				CurrentShaderMapId.Add(CompilingId);
-				GOpenColorIOShaderCompilationManager.FinishCompilation(*FriendlyName, CurrentShaderMapId);
+				GOpenColorIOShaderCompilationManager.FinishCompilation(*NewContent->FriendlyName, CurrentShaderMapId);
 			}
 		}
 	}
 }
 
-FShader* FOpenColorIOShaderMap::ProcessCompilationResultsForSingleJob(FShaderCompileJob& CurrentJob, const FSHAHash& InShaderMapHash)
+FShader* FOpenColorIOShaderMap::ProcessCompilationResultsForSingleJob(FShaderCompileJob& CurrentJob, const FSHAHash& InShaderMapHash, FShaderMapResourceBuilder& InResourceBuilder)
 {
 	check(CurrentJob.Id == CompilingId);
 
@@ -477,15 +484,13 @@ FShader* FOpenColorIOShaderMap::ProcessCompilationResultsForSingleJob(FShaderCom
 
 	FOpenColorIOShaderType* OpenColorIOShaderType = CurrentJob.ShaderType->GetOpenColorIOShaderType();
 	check(OpenColorIOShaderType);
-	Shader = OpenColorIOShaderType->FinishCompileShader(InShaderMapHash, CurrentJob, FriendlyName);
+	Shader = OpenColorIOShaderType->FinishCompileShader(InShaderMapHash, CurrentJob, GetContent()->FriendlyName, InResourceBuilder);
 	bCompiledSuccessfully = CurrentJob.bSucceeded;
 
 	FOpenColorIOPixelShader *OpenColorIOShader = static_cast<FOpenColorIOPixelShader*>(Shader);
 	check(Shader);
-	check(!HasShader(OpenColorIOShaderType, /* PermutationId = */ 0));
-	AddShader(OpenColorIOShaderType, /* PermutationId = */ 0, Shader);
-
-	return Shader;
+	check(!GetContent()->HasShader(OpenColorIOShaderType, /* PermutationId = */ 0));
+	return GetMutableContent()->FindOrAddShader(Shader);
 }
 
 bool FOpenColorIOShaderMap::ProcessCompilationResults(const TArray<TSharedRef<FShaderCommonCompileJob, ESPMode::ThreadSafe>>& InCompilationResults, int32& InOutJobIndex, float& InOutTimeBudget)
@@ -495,11 +500,12 @@ bool FOpenColorIOShaderMap::ProcessCompilationResults(const TArray<TSharedRef<FS
 	double StartTime = FPlatformTime::Seconds();
 
 	FSHAHash ShaderMapHash;
-	ShaderMapId.GetOpenColorIOHash(ShaderMapHash);
+	GetContent()->ShaderMapId.GetOpenColorIOHash(ShaderMapHash);
 
+	FShaderMapResourceBuilder ResourceBuilder(GetResourceCode());
 	do
 	{
-		ProcessCompilationResultsForSingleJob(StaticCastSharedRef<FShaderCompileJob>(InCompilationResults[InOutJobIndex]).Get(), ShaderMapHash);
+		ProcessCompilationResultsForSingleJob(StaticCastSharedRef<FShaderCompileJob>(InCompilationResults[InOutJobIndex]).Get(), ShaderMapHash, ResourceBuilder);
 
 		InOutJobIndex++;
 		
@@ -511,6 +517,8 @@ bool FOpenColorIOShaderMap::ProcessCompilationResults(const TArray<TSharedRef<FS
 
 	if (InOutJobIndex == InCompilationResults.Num())
 	{
+		FinalizeContent();
+
 		SaveToDerivedDataCache();
 		// The shader map can now be used on the rendering thread
 		bCompilationFinalized = true;
@@ -543,7 +551,7 @@ bool FOpenColorIOShaderMap::TryToAddToExistingCompilationTask(FOpenColorIOTransf
 bool FOpenColorIOShaderMap::IsOpenColorIOShaderComplete(const FOpenColorIOTransformResource* InColorTransform, const FOpenColorIOShaderType* InShaderType, bool bSilent)
 {
 	// If we should cache this color transform, it's incomplete if the shader is missing
-	if (ShouldCacheOpenColorIOShader(InShaderType, Platform, InColorTransform) &&	!HasShader((FShaderType*)InShaderType, /* PermutationId = */ 0))
+	if (ShouldCacheOpenColorIOShader(InShaderType, GetShaderPlatform(), InColorTransform) && !GetContent()->HasShader((FShaderType*)InShaderType, /* PermutationId = */ 0))
 	{
 		if (!bSilent)
 		{
@@ -584,6 +592,7 @@ bool FOpenColorIOShaderMap::IsComplete(const FOpenColorIOTransformResource* InCo
 
 void FOpenColorIOShaderMap::LoadMissingShadersFromMemory(const FOpenColorIOTransformResource* InColorTransform)
 {
+#if 0
 	// Make sure we are operating on a referenced shader map or the below Find will cause this shader map to be deleted,
 	// Since it creates a temporary ref counted pointer.
 	check(NumRefs > 0);
@@ -613,11 +622,12 @@ void FOpenColorIOShaderMap::LoadMissingShadersFromMemory(const FOpenColorIOTrans
 			}
 		}
 	}
+#endif
 }
 
-void FOpenColorIOShaderMap::GetShaderList(TMap<FShaderId, FShader*>& OutShaders) const
+void FOpenColorIOShaderMap::GetShaderList(TMap<FShaderId, TShaderRef<FShader>>& OutShaders) const
 {
-	TShaderMap<FOpenColorIOShaderType>::GetShaderList(OutShaders);
+	GetContent()->GetShaderList(*this, FSHAHash(), OutShaders);
 }
 
 /**
@@ -625,25 +635,12 @@ void FOpenColorIOShaderMap::GetShaderList(TMap<FShaderId, FShader*>& OutShaders)
  */
 void FOpenColorIOShaderMap::Register(EShaderPlatform InShaderPlatform)
 {
-	if (Platform == InShaderPlatform)
-	{
-		for (auto KeyValue : GetShaders())
-		{
-			FShader* Shader = KeyValue.Value;
-			if (Shader)
-			{
-				Shader->BeginInitializeResources();
-			}
-		}
-	}
-
 	if (!bRegistered)
 	{
 		INC_DWORD_STAT(STAT_Shaders_NumShaderMaps);
-		INC_DWORD_STAT_BY(STAT_Shaders_ShaderMapMemory, GetSizeBytes());
 	}
 
-	GIdToOpenColorIOShaderMap[Platform].Add(ShaderMapId,this);
+	GIdToOpenColorIOShaderMap[GetShaderPlatform()].Add(GetContent()->ShaderMapId,this);
 	bRegistered = true;
 }
 
@@ -661,9 +658,8 @@ void FOpenColorIOShaderMap::Release()
 		if (bRegistered)
 		{
 			DEC_DWORD_STAT(STAT_Shaders_NumShaderMaps);
-			DEC_DWORD_STAT_BY(STAT_Shaders_ShaderMapMemory, GetSizeBytes());
 
-			GIdToOpenColorIOShaderMap[Platform].Remove(ShaderMapId);
+			GIdToOpenColorIOShaderMap[GetShaderPlatform()].Remove(GetContent()->ShaderMapId);
 			bRegistered = false;
 		}
 
@@ -674,8 +670,6 @@ void FOpenColorIOShaderMap::Release()
 }
 
 FOpenColorIOShaderMap::FOpenColorIOShaderMap() :
-	TShaderMap<FOpenColorIOShaderType>(SP_NumPlatforms),
-	Platform(SP_NumPlatforms),
 	CompilingId(1),
 	NumRefs(0),
 	bDeletedThroughDeferredCleanup(false),
@@ -700,53 +694,21 @@ FOpenColorIOShaderMap::~FOpenColorIOShaderMap()
  * Removes all entries in the cache with exceptions based on a shader type
  * @param ShaderType - The shader type to flush
  */
-void FOpenColorIOShaderMap::FlushShadersByShaderType(FShaderType* InShaderType)
+void FOpenColorIOShaderMap::FlushShadersByShaderType(const FShaderType* InShaderType)
 {
 	if (InShaderType->GetOpenColorIOShaderType())
 	{
-		RemoveShaderTypePermutaion(InShaderType->GetOpenColorIOShaderType(), /* PermutationId = */ 0);	
+		GetMutableContent()->RemoveShaderTypePermutaion(InShaderType->GetOpenColorIOShaderType(), /* PermutationId = */ 0);	
 	}
 }
 
+
 void FOpenColorIOShaderMap::Serialize(FArchive& Ar, bool bInlineShaderResources)
 {
 	// Note: This is saved to the DDC, not into packages (except when cooked)
 	// Backwards compatibility therefore will not work based on the version of Ar
 	// Instead, just bump OPENCOLORIO_DERIVEDDATA_VER
-
-	ShaderMapId.Serialize(Ar);
-
-	// serialize the platform enum as a uint8
-	int32 TempPlatform = (int32)Platform;
-	Ar << TempPlatform;
-	Platform = (EShaderPlatform)TempPlatform;
-
-	Ar << FriendlyName;
-
-	OpenColorIOCompilationOutput.Serialize(Ar);
-
-	if (Ar.IsSaving())
-	{
-		TShaderMap<FOpenColorIOShaderType>::SerializeInline(Ar, bInlineShaderResources, false, false);
-		RegisterSerializedShaders(false);
-	}
-
-	if (Ar.IsLoading())
-	{
-		TShaderMap<FOpenColorIOShaderType>::SerializeInline(Ar, bInlineShaderResources, false, false);
-	}
-}
-
-void FOpenColorIOShaderMap::RegisterSerializedShaders(bool bCooked)
-{
-	check(IsInGameThread());
-
-	TShaderMap<FOpenColorIOShaderType>::RegisterSerializedShaders(bCooked);
-}
-
-void FOpenColorIOShaderMap::DiscardSerializedShaders()
-{
-	TShaderMap<FOpenColorIOShaderType>::DiscardSerializedShaders();
+	Super::Serialize(Ar, bInlineShaderResources, false);
 }
 
 void FOpenColorIOShaderMap::RemovePendingColorTransform(FOpenColorIOTransformResource* InColorTransform)
diff --git a/Engine/Plugins/Compositing/OpenColorIO/Source/OpenColorIO/Private/OpenColorIOShared.cpp b/Engine/Plugins/Compositing/OpenColorIO/Source/OpenColorIO/Private/OpenColorIOShared.cpp
--- a/Engine/Plugins/Compositing/OpenColorIO/Source/OpenColorIO/Private/OpenColorIOShared.cpp
+++ b/Engine/Plugins/Compositing/OpenColorIO/Source/OpenColorIO/Private/OpenColorIOShared.cpp
@@ -20,6 +20,10 @@
 #include "UObject/UObjectHash.h"
 #include "UObject/UObjectIterator.h"
 
+IMPLEMENT_TYPE_LAYOUT(FOpenColorIOCompilationOutput);
+IMPLEMENT_TYPE_LAYOUT(FOpenColorIOShaderMapId);
+IMPLEMENT_TYPE_LAYOUT(FOpenColorIOShaderMapContent);
+
 FOpenColorIOTransformResource::~FOpenColorIOTransformResource()
 {
 	FOpenColorIOShaderMap::RemovePendingColorTransform(this);
@@ -105,14 +109,6 @@ OPENCOLORIO_API void FOpenColorIOTransformResource::GetShaderMapId(EShaderPlatfo
 	}
 }
 
-void FOpenColorIOTransformResource::RegisterShaderMap()
-{
-	if (GameThreadShaderMap)
-	{
-		GameThreadShaderMap->RegisterSerializedShaders(false);
-	}
-}
-
 void FOpenColorIOTransformResource::ReleaseShaderMap()
 {
 	if (GameThreadShaderMap)
@@ -131,11 +127,7 @@ void FOpenColorIOTransformResource::ReleaseShaderMap()
 void FOpenColorIOTransformResource::DiscardShaderMap()
 {
 	check(RenderingThreadShaderMap == nullptr);
-	if (GameThreadShaderMap)
-	{
-		GameThreadShaderMap->DiscardSerializedShaders();
-		GameThreadShaderMap = nullptr;
-	}
+	GameThreadShaderMap = nullptr;
 }
 
 void FOpenColorIOTransformResource::SerializeShaderMap(FArchive& Ar)
@@ -185,10 +177,6 @@ void FOpenColorIOTransformResource::SerializeShaderMap(FArchive& Ar)
 				{
 					GameThreadShaderMap = RenderingThreadShaderMap = LoadedShaderMap;
 				}
-				else
-				{
-					LoadedShaderMap->DiscardSerializedShaders();
-				}
 			}
 		}
 	}
@@ -351,24 +339,24 @@ void FOpenColorIOTransformResource::FinishCompilation()
 #endif
 }
 
-OPENCOLORIO_API  FOpenColorIOPixelShader* FOpenColorIOTransformResource::GetShader() const
+OPENCOLORIO_API  TShaderRef<FOpenColorIOPixelShader> FOpenColorIOTransformResource::GetShader() const
 {
 	check(!GIsThreadedRendering || !IsInGameThread());
 	if (!GIsEditor || RenderingThreadShaderMap)
 	{
 		return RenderingThreadShaderMap->GetShader<FOpenColorIOPixelShader>();
 	}
-	return nullptr;
+	return TShaderRef<FOpenColorIOPixelShader>();
 };
 
-OPENCOLORIO_API  FOpenColorIOPixelShader* FOpenColorIOTransformResource::GetShaderGameThread() const
+OPENCOLORIO_API  TShaderRef<FOpenColorIOPixelShader> FOpenColorIOTransformResource::GetShaderGameThread() const
 {
 	if (GameThreadShaderMap)
 	{
 		return GameThreadShaderMap->GetShader<FOpenColorIOPixelShader>();
 	}
 
-	return nullptr;
+	return TShaderRef<FOpenColorIOPixelShader>();
 };
 
 void FOpenColorIOTransformResource::GetShaderMapIDsWithUnfinishedCompilation(TArray<int32>& OutShaderMapIds)
@@ -443,7 +431,7 @@ void FOpenColorIOShaderMapId::SetShaderDependencies(const TArray<FShaderType*>&
 			if (ShaderType != nullptr)
 			{
 				FShaderTypeDependency Dependency;
-				Dependency.ShaderType = ShaderType;
+				Dependency.ShaderTypeName = ShaderType->GetHashedName();
 				Dependency.SourceHash = ShaderType->GetSourceHash(InShaderPlatform);
 				ShaderTypeDependencies.Add(Dependency);
 			}
@@ -456,7 +444,7 @@ bool FOpenColorIOShaderMapId::ContainsShaderType(const FShaderType* ShaderType)
 {
 	for (int32 TypeIndex = 0; TypeIndex < ShaderTypeDependencies.Num(); TypeIndex++)
 	{
-		if (ShaderTypeDependencies[TypeIndex].ShaderType == ShaderType)
+		if (ShaderTypeDependencies[TypeIndex].ShaderTypeName == ShaderType->GetHashedName())
 		{
 			return true;
 		}
diff --git a/Engine/Plugins/Compositing/OpenColorIO/Source/OpenColorIO/Public/OpenColorIOShader.h b/Engine/Plugins/Compositing/OpenColorIO/Source/OpenColorIO/Public/OpenColorIOShader.h
--- a/Engine/Plugins/Compositing/OpenColorIO/Source/OpenColorIO/Public/OpenColorIOShader.h
+++ b/Engine/Plugins/Compositing/OpenColorIO/Source/OpenColorIO/Public/OpenColorIOShader.h
@@ -24,56 +24,48 @@ namespace OpenColorIOShader
 	static const uint32 Lut3dEdgeLength = 32;
 }
 
-
+class FOpenColorIOShader : public FShader
+{
+	DECLARE_TYPE_LAYOUT(FOpenColorIOShader, NonVirtual);
+public:
+	FOpenColorIOShader() = default;
+	FOpenColorIOShader(const FOpenColorIOShader::CompiledShaderInitializerType& Initializer) : FShader(Initializer) {}
+};
 
 /** Base class of all shaders that need OpenColorIO pixel shader parameters. */
-class OPENCOLORIO_API FOpenColorIOPixelShader : public FShader
+class OPENCOLORIO_API FOpenColorIOPixelShader : public FOpenColorIOShader
 {
 public:
 	DECLARE_SHADER_TYPE(FOpenColorIOPixelShader, OpenColorIO);
 
+	using FPermutationParameters = FOpenColorIOShaderPermutationParameters;
+
 	FOpenColorIOPixelShader()
 	{
 	}
 
-	static bool ShouldCompilePermutation(EShaderPlatform Platform, const FOpenColorIOTransformResource*  InColorTransform)
+	static bool ShouldCompilePermutation(const FOpenColorIOShaderPermutationParameters& Parameters)
 	{
-		return IsFeatureLevelSupported(Platform, ERHIFeatureLevel::ES2);
+		return IsFeatureLevelSupported(Parameters.Platform, ERHIFeatureLevel::ES2);
 	}
 
-
 	FOpenColorIOPixelShader(const FOpenColorIOShaderType::CompiledShaderInitializerType& Initializer);
 
-	typedef void (*ModifyCompilationEnvironmentType)(EShaderPlatform, const FOpenColorIOTransformResource* , FShaderCompilerEnvironment&);
-
-	static void ModifyCompilationEnvironment(EShaderPlatform InPlatform, const FOpenColorIOTransformResource*  InColorTransform, FShaderCompilerEnvironment& OutEnvironment)
-	{
-	}
-	
-	static bool ValidateCompiledResult(EShaderPlatform InPlatform, const FShaderParameterMap& InParameterMap, TArray<FString>& OutError)
-	{
-		return true;
-	}
-
 	void SetParameters(FRHICommandList& InRHICmdList, FTextureResource* InInputTexture);
 	void SetLUTParameter(FRHICommandList& InRHICmdList, FTextureResource* InLUT3dResource);
 	
 	// Bind parameters
 	void BindParams(const FShaderParameterMap &ParameterMap);
 
-	// FShader interface.
-	virtual bool Serialize(FArchive& Ar) override;
-	virtual uint32 GetAllocatedSize() const override;
-
 protected:
-	FShaderResourceParameter InputTexture;
-	FShaderResourceParameter InputTextureSampler;
+	LAYOUT_FIELD(FShaderResourceParameter, InputTexture)
+	LAYOUT_FIELD(FShaderResourceParameter, InputTextureSampler)
 	
-	FShaderResourceParameter OCIO3dTexture;
-	FShaderResourceParameter OCIO3dTextureSampler;
+	LAYOUT_FIELD(FShaderResourceParameter, OCIO3dTexture)
+	LAYOUT_FIELD(FShaderResourceParameter, OCIO3dTextureSampler)
 
 private:
-	FString						DebugDescription;
+	LAYOUT_FIELD(FMemoryImageString, DebugDescription)
 };
 
 
diff --git a/Engine/Plugins/Compositing/OpenColorIO/Source/OpenColorIO/Public/OpenColorIOShaderType.h b/Engine/Plugins/Compositing/OpenColorIO/Source/OpenColorIO/Public/OpenColorIOShaderType.h
--- a/Engine/Plugins/Compositing/OpenColorIO/Source/OpenColorIO/Public/OpenColorIOShaderType.h
+++ b/Engine/Plugins/Compositing/OpenColorIO/Source/OpenColorIO/Public/OpenColorIOShaderType.h
@@ -30,6 +30,15 @@ class FUniformExpressionSet;
 /** Called for every OpenColorIO shader to update the appropriate stats. */
 extern void UpdateOpenColorIOShaderCompilingStats(const FOpenColorIOTransformResource* InShader);
 
+struct FOpenColorIOShaderPermutationParameters : public FShaderPermutationParameters
+{
+	const FOpenColorIOTransformResource* Transform;
+
+	FOpenColorIOShaderPermutationParameters(EShaderPlatform InPlatform, const FOpenColorIOTransformResource* InTransform)
+		: FShaderPermutationParameters(InPlatform)
+		, Transform(InTransform)
+	{}
+};
 
 /**
  * A shader meta type for OpenColorIO-linked shaders.
@@ -45,20 +54,17 @@ public:
 			FShaderType* InType,
 			int32 InPermutationId,
 			const FShaderCompilerOutput& CompilerOutput,
-			FShaderResource* InResource,
+			int32 InResourceIndex,
 			const FSHAHash& InOCIOShaderMapHash,
 			const FString& InDebugDescription
 			)
-		: FGlobalShaderType::CompiledShaderInitializerType(InType,InPermutationId,CompilerOutput,InResource, InOCIOShaderMapHash,nullptr,nullptr)
+		: FGlobalShaderType::CompiledShaderInitializerType(InType,InPermutationId,CompilerOutput, InResourceIndex, InOCIOShaderMapHash,nullptr,nullptr)
 		, DebugDescription(InDebugDescription)
 		{}
 	};
-	typedef FShader* (*ConstructCompiledType)(const CompiledShaderInitializerType&);
-	typedef bool (*ShouldCompilePermutationType)(EShaderPlatform,const FOpenColorIOTransformResource*);
-	typedef bool(*ValidateCompiledResultType)(EShaderPlatform, const FShaderParameterMap&, TArray<FString>&);
-	typedef void (*ModifyCompilationEnvironmentType)(EShaderPlatform,const FOpenColorIOTransformResource*, FShaderCompilerEnvironment&);
 
 	FOpenColorIOShaderType(
+		FTypeLayoutDesc& InTypeLayout,
 		const TCHAR* InName,
 		const TCHAR* InSourceFilename,
 		const TCHAR* InFunctionName,
@@ -68,13 +74,18 @@ public:
 		ConstructCompiledType InConstructCompiledRef,
 		ModifyCompilationEnvironmentType InModifyCompilationEnvironmentRef,
 		ShouldCompilePermutationType InShouldCompilePermutationRef,
-		ValidateCompiledResultType InValidateCompiledResultRef
+		ValidateCompiledResultType InValidateCompiledResultRef,
+		uint32 InTypeSize,
+		const FShaderParametersMetadata* InRootParametersMetadata = nullptr
 		):
-		FShaderType(EShaderTypeForDynamicCast::OCIO, InName, InSourceFilename, InFunctionName, SF_Pixel, InTotalPermutationCount, InConstructSerializedRef, nullptr),
-		ConstructCompiledRef(InConstructCompiledRef),
-		ShouldCompilePermutationRef(InShouldCompilePermutationRef),
-		ValidateCompiledResultRef(InValidateCompiledResultRef),
-		ModifyCompilationEnvironmentRef(InModifyCompilationEnvironmentRef)
+		FShaderType(EShaderTypeForDynamicCast::OCIO, InTypeLayout, InName, InSourceFilename, InFunctionName, SF_Pixel, InTotalPermutationCount,
+			InConstructSerializedRef,
+			InConstructCompiledRef,
+			InModifyCompilationEnvironmentRef,
+			InShouldCompilePermutationRef,
+			InValidateCompiledResultRef,
+			InTypeSize,
+			InRootParametersMetadata)
 	{
 		check(InTotalPermutationCount == 1);
 	}
@@ -99,7 +110,8 @@ public:
 	FShader* FinishCompileShader(
 		const FSHAHash& InOCIOShaderMapHash,
 		const FShaderCompileJob& CurrentJob,
-		const FString& InDebugDescription
+		const FString& InDebugDescription,
+		FShaderMapResourceBuilder& ResourceBuilder
 		);
 
 	/**
@@ -108,28 +120,20 @@ public:
 	 * @param InColorTransform - The color transform to check.
 	 * @return True if this shader type should be cached.
 	 */
-	bool ShouldCache(EShaderPlatform InPlatform,const FOpenColorIOTransformResource* InColorTransform) const
+	bool ShouldCache(EShaderPlatform InPlatform, const FOpenColorIOTransformResource* InColorTransform) const
 	{
-		return (*ShouldCompilePermutationRef)(InPlatform, InColorTransform);
+		return ShouldCompilePermutation(FOpenColorIOShaderPermutationParameters(InPlatform, InColorTransform));
 	}
 
 
 protected:
-
 	/**
 	 * Sets up the environment used to compile an instance of this shader type.
 	 * @param InPlatform - Platform to compile for.
 	 * @param OutEnvironment - The shader compile environment that the function modifies.
 	 */
-	void SetupCompileEnvironment(EShaderPlatform InPlatform, const FOpenColorIOTransformResource* InColorTransform, FShaderCompilerEnvironment& OutEnvironment)
+	void SetupCompileEnvironment(EShaderPlatform InPlatform, const FOpenColorIOTransformResource* InColorTransform, FShaderCompilerEnvironment& OutEnvironment) const
 	{
-		// Allow the shader type to modify its compile environment.
-		(*ModifyCompilationEnvironmentRef)(InPlatform, InColorTransform, OutEnvironment);
+		ModifyCompilationEnvironment(FOpenColorIOShaderPermutationParameters(InPlatform, InColorTransform), OutEnvironment);
 	}
-
-private:
-	ConstructCompiledType ConstructCompiledRef;
-	ShouldCompilePermutationType ShouldCompilePermutationRef;
-	ValidateCompiledResultType ValidateCompiledResultRef;
-	ModifyCompilationEnvironmentType ModifyCompilationEnvironmentRef;
 };
diff --git a/Engine/Plugins/Compositing/OpenColorIO/Source/OpenColorIO/Public/OpenColorIOShared.h b/Engine/Plugins/Compositing/OpenColorIO/Source/OpenColorIO/Public/OpenColorIOShared.h
--- a/Engine/Plugins/Compositing/OpenColorIO/Source/OpenColorIO/Public/OpenColorIOShared.h
+++ b/Engine/Plugins/Compositing/OpenColorIO/Source/OpenColorIO/Public/OpenColorIOShared.h
@@ -31,28 +31,27 @@ class FOpenColorIOShaderMapId;
 /** Stores outputs from the color transform compile that need to be saved. */
 class FOpenColorIOCompilationOutput
 {
+	DECLARE_TYPE_LAYOUT(FOpenColorIOCompilationOutput, NonVirtual);
 public:
 	FOpenColorIOCompilationOutput()
 	{}
-
-	OPENCOLORIO_API void Serialize(FArchive& Ar)
-	{}
 };
 
 /** Contains all the information needed to uniquely identify a FOpenColorIOShaderMapID. */
 class FOpenColorIOShaderMapId
 {
+	DECLARE_TYPE_LAYOUT(FOpenColorIOShaderMapId, NonVirtual);
 public:
 	/** Feature level that the shader map is going to be compiled for.  */
-	ERHIFeatureLevel::Type FeatureLevel;
+	LAYOUT_FIELD(ERHIFeatureLevel::Type, FeatureLevel);
 
 	/**
 	 * The hash computed from the shader code generated by the library.
 	 */
-	FString ShaderCodeHash;
+	LAYOUT_FIELD(FMemoryImageString, ShaderCodeHash);
 	
 	/** Shader types of shaders that are inlined in this shader map in the DDC. */
-	TArray<FShaderTypeDependency> ShaderTypeDependencies;
+	LAYOUT_FIELD(TMemoryImageArray<FShaderTypeDependency>, ShaderTypeDependencies);
 
 public:
 	FOpenColorIOShaderMapId()
@@ -64,7 +63,7 @@ public:
 
 	void SetShaderDependencies(const TArray<FShaderType*>& InShaderTypes, EShaderPlatform InShaderPlatform);
 
-	void Serialize(FArchive& Ar);
+	//void Serialize(FArchive& Ar);
 
 	friend uint32 GetTypeHash(const FOpenColorIOShaderMapId& Ref)
 	{
@@ -99,14 +98,32 @@ public:
 	bool ContainsShaderType(const FShaderType* ShaderType) const;
 };
 
+class FOpenColorIOShaderMapContent : public FShaderMapContent
+{
+	using Super = FShaderMapContent;
+	friend class FOpenColorIOShaderMap;
+	DECLARE_TYPE_LAYOUT(FOpenColorIOShaderMapContent, NonVirtual);
+private:
+	explicit FOpenColorIOShaderMapContent(EShaderPlatform InPlatform) : Super(InPlatform) {}
+
+	/** The ColorTransform's user friendly name, typically the source and destination color spaces. */
+	LAYOUT_FIELD(FMemoryImageString, FriendlyName);
+
+	/** The static parameter set that this shader map was compiled with */
+	LAYOUT_FIELD(FOpenColorIOShaderMapId, ShaderMapId);
+
+	/** Shader compilation output */
+	LAYOUT_FIELD(FOpenColorIOCompilationOutput, OpenColorIOCompilationOutput);
+};
 
 
 /**
  * The shaders required to apply an OCIO color transform.
  */
-class FOpenColorIOShaderMap : public TShaderMap<FOpenColorIOShaderType>, public FDeferredCleanupInterface
+class FOpenColorIOShaderMap : public TShaderMap<FOpenColorIOShaderMapContent, FShaderMapPointerTable>, public FDeferredCleanupInterface
 {
 public:
+	using Super = TShaderMap<FOpenColorIOShaderMapContent, FShaderMapPointerTable>;
 
 	/**
 	 * Finds the shader map for a color transform.
@@ -116,6 +133,10 @@ public:
 	 */
 	static FOpenColorIOShaderMap* FindId(const FOpenColorIOShaderMapId& InShaderMapId, EShaderPlatform InPlatform);
 
+	// ShaderMap interface
+	template<typename ShaderType> TShaderRef<ShaderType> GetShader() const { return TShaderRef<ShaderType>(GetContent()->GetShader<ShaderType>(), *this); }
+	TShaderRef<FShader> GetShader(FShaderType* ShaderType) const { return TShaderRef<FShader>(GetContent()->GetShader(ShaderType), *this); }
+
 	/**
 	 * Attempts to load the shader map for the given color transform from the Derived Data Cache.
 	 * If InOutShaderMap is valid, attempts to load the individual missing shaders instead.
@@ -165,7 +186,7 @@ public:
 	bool TryToAddToExistingCompilationTask(FOpenColorIOTransformResource* InColorTransform);
 
 	/** Builds a list of the shaders in a shader map. */
-	OPENCOLORIO_API  void GetShaderList(TMap<FShaderId, FShader*>& OutShaders) const;
+	OPENCOLORIO_API  void GetShaderList(TMap<FShaderId, TShaderRef<FShader>>& OutShaders) const;
 
 	/** Registers a OpenColorIO shader map in the global map so it can be used by OpenColorIO ColorTransform. */
 	void Register(EShaderPlatform InShaderPlatform);
@@ -178,7 +199,7 @@ public:
 	 * Removes all entries in the cache with exceptions based on a shader type
 	 * @param ShaderType - The shader type to flush
 	 */
-	void FlushShadersByShaderType(FShaderType* ShaderType);
+	void FlushShadersByShaderType(const FShaderType* ShaderType);
 
 	/** Removes a ColorTransform from OpenColorIOShaderMapsBeingCompiled. */
 	OPENCOLORIO_API static void RemovePendingColorTransform(FOpenColorIOTransformResource* InColorTransform);
@@ -193,22 +214,16 @@ public:
 	void SaveToDerivedDataCache();
 
 	/** Registers all shaders that have been loaded in Serialize */
-	virtual void RegisterSerializedShaders(bool bCooked) override;
-	virtual void DiscardSerializedShaders() override;
-
-	/** Computes the memory used by this shader map without counting the shaders themselves. */
-	uint32 GetSizeBytes() const
-	{
-		return sizeof(*this) + FriendlyName.GetAllocatedSize();
-	}
+	//virtual void RegisterSerializedShaders(bool bCooked) override;
+	//virtual void DiscardSerializedShaders() override;
 
 	// Accessors.
-	const FOpenColorIOShaderMapId& GetShaderMapId() const { return ShaderMapId; }
-	EShaderPlatform GetShaderPlatform() const	{ return Platform; }
-	const FString& GetFriendlyName() const		{ return FriendlyName; }
-	uint32 GetCompilingId() const				{ return CompilingId; }
-	bool IsCompilationFinalized() const			{ return bCompilationFinalized; }
-	bool CompiledSuccessfully() const			{ return bCompiledSuccessfully; }
+	const FOpenColorIOShaderMapId& GetShaderMapId() const	{ return GetContent()->ShaderMapId; }
+	EShaderPlatform GetShaderPlatform() const				{ return GetContent()->GetShaderPlatform(); }
+	const FMemoryImageString& GetFriendlyName() const		{ return GetContent()->FriendlyName; }
+	uint32 GetCompilingId() const							{ return CompilingId; }
+	bool IsCompilationFinalized() const						{ return bCompilationFinalized; }
+	bool CompiledSuccessfully() const						{ return bCompiledSuccessfully; }
 
 	bool IsValid() const
 	{
@@ -235,18 +250,6 @@ private:
 	 */
 	static TArray<FOpenColorIOShaderMap*> AllOpenColorIOShaderMaps;
 
-	/** The ColorTransform's user friendly name, typically the source and destination color spaces. */
-	FString FriendlyName;
-
-	/** The platform this shader map was compiled with */
-	EShaderPlatform Platform;
-
-	/** The static parameter set that this shader map was compiled with */
-	FOpenColorIOShaderMapId ShaderMapId;
-
-	/** Shader compilation output */
-	FOpenColorIOCompilationOutput OpenColorIOCompilationOutput;
-
 	/** Next value for CompilingId. */
 	static uint32 NextCompilingId;
 
@@ -275,7 +278,9 @@ private:
 	/** Indicates whether the shader map should be stored in the shader cache. */
 	uint32 bIsPersistent : 1;
 
-	FShader* ProcessCompilationResultsForSingleJob(FShaderCompileJob& InSingleJob, const FSHAHash& InShaderMapHash);
+	uint32 bHasFrozenContent : 1;
+
+	FShader* ProcessCompilationResultsForSingleJob(FShaderCompileJob& InSingleJob, const FSHAHash& InShaderMapHash, FShaderMapResourceBuilder& InResourceBuilder);
 
 	bool IsOpenColorIOShaderComplete(const FOpenColorIOTransformResource* InColorTransform, const FOpenColorIOShaderType* InShaderType, bool bSilent);
 
@@ -376,8 +381,6 @@ public:
 		return GameThreadShaderMap;
 	}
 
-	OPENCOLORIO_API void RegisterShaderMap();
-
 	/** Note: SetRenderingThreadShaderMap must also be called with the same value, but from the rendering thread. */
 	void SetGameThreadShaderMap(FOpenColorIOShaderMap* InShaderMap)
 	{
@@ -427,8 +430,8 @@ public:
 	
 	FString ShaderCode;
 
-	OPENCOLORIO_API  FOpenColorIOPixelShader* GetShader() const;
-	OPENCOLORIO_API  FOpenColorIOPixelShader* GetShaderGameThread() const;
+	OPENCOLORIO_API  TShaderRef<FOpenColorIOPixelShader> GetShader() const;
+	OPENCOLORIO_API  TShaderRef<FOpenColorIOPixelShader> GetShaderGameThread() const;
 	
 	bool IsSame(const FOpenColorIOShaderMapId& InId) const;
 protected:
diff --git a/Engine/Plugins/Enterprise/DataprepEditor/Source/DataprepCore/Private/DataprepCorePrivateUtils.cpp b/Engine/Plugins/Enterprise/DataprepEditor/Source/DataprepCore/Private/DataprepCorePrivateUtils.cpp
--- a/Engine/Plugins/Enterprise/DataprepEditor/Source/DataprepCore/Private/DataprepCorePrivateUtils.cpp
+++ b/Engine/Plugins/Enterprise/DataprepEditor/Source/DataprepCore/Private/DataprepCorePrivateUtils.cpp
@@ -207,7 +207,7 @@ void DataprepCorePrivateUtils::ClearAssets(const TArray<TWeakObjectPtr<UObject>>
 		if(UStaticMesh* StaticMesh = Cast<UStaticMesh>(ObjectPtr.Get()))
 		{
 			StaticMesh->PreEditChange( nullptr );
-			StaticMesh->RenderData.Reset();
+			StaticMesh->RenderData.Reset(nullptr);
 		}
 	}
 }
diff --git a/Engine/Plugins/Enterprise/DataprepEditor/Source/DataprepEditor/Private/Widgets/SDataprepEditorViewport.cpp b/Engine/Plugins/Enterprise/DataprepEditor/Source/DataprepEditor/Private/Widgets/SDataprepEditorViewport.cpp
--- a/Engine/Plugins/Enterprise/DataprepEditor/Source/DataprepEditor/Private/Widgets/SDataprepEditorViewport.cpp
+++ b/Engine/Plugins/Enterprise/DataprepEditor/Source/DataprepEditor/Private/Widgets/SDataprepEditorViewport.cpp
@@ -380,7 +380,7 @@ void SDataprepEditorViewport::UpdateScene()
 					{
 						const FText Message = FText::Format( LOCTEXT("UpdateMeshes_AddOneComponent", "Adding {0} ..."), FText::FromString( SceneMeshComponent->GetOwner()->GetActorLabel() ) );
 						SubSlowTask.EnterProgressFrame( 1.0f, Message );
-						
+
 						UCustomStaticMeshComponent* PreviewMeshComponent = NewObject< UCustomStaticMeshComponent >( PreviewActor.Get(), NAME_None, RF_Transient );
 						if (GEditor->PreviewPlatform.GetEffectivePreviewFeatureLevel() <= ERHIFeatureLevel::ES3_1)
 						{
@@ -1497,7 +1497,7 @@ void FDataprepEditorViewportCommands::RegisterCommands()
 	UI_COMMAND(ApplyXRayMaterial, "XRay", "Use XRay material to render meshes.", EUserInterfaceActionType::RadioButton, FInputChord());
 	UI_COMMAND(ApplyPerMeshMaterial, "MultiColored", "Assign a different color for each rendered mesh.", EUserInterfaceActionType::RadioButton, FInputChord());
 	UI_COMMAND(ApplyReflectionMaterial, "ReflectionLines", "Use reflective material to show lines of reflection.", EUserInterfaceActionType::RadioButton, FInputChord());
-	
+
 	// Selection Mode
 	UI_COMMAND(ApplyOutlineSelection, "Outline", "Outline selected meshes with a colored contour.", EUserInterfaceActionType::RadioButton, FInputChord());
 	UI_COMMAND(ApplyXRaySelection, "XRay", "Use XRay material on non selected meshes.", EUserInterfaceActionType::RadioButton, FInputChord());
@@ -1519,7 +1519,7 @@ FPrimitiveSceneProxy* UCustomStaticMeshComponent::CreateSceneProxy()
 		return nullptr;
 	}
 
-	const TIndirectArray<FStaticMeshLODResources>& LODResources = GetStaticMesh()->RenderData->LODResources;
+	const FStaticMeshLODResourcesArray& LODResources = GetStaticMesh()->RenderData->LODResources;
 	if (LODResources.Num() == 0	|| LODResources[FMath::Clamp<int32>(GetStaticMesh()->MinLOD.Default, 0, LODResources.Num()-1)].VertexBuffers.StaticMeshVertexBuffer.GetNumVertices() == 0)
 	{
 		return nullptr;
diff --git a/Engine/Plugins/Enterprise/DatasmithImporter/Source/DatasmithImporter/Private/DatasmithConsumer.cpp b/Engine/Plugins/Enterprise/DatasmithImporter/Source/DatasmithImporter/Private/DatasmithConsumer.cpp
--- a/Engine/Plugins/Enterprise/DatasmithImporter/Source/DatasmithImporter/Private/DatasmithConsumer.cpp
+++ b/Engine/Plugins/Enterprise/DatasmithImporter/Source/DatasmithImporter/Private/DatasmithConsumer.cpp
@@ -770,7 +770,7 @@ namespace DatasmithConsumerUtils
 					// Clean up static meshes which have incomplete render data.
 					if(StaticMesh->RenderData.IsValid() && !StaticMesh->RenderData->IsInitialized())
 					{
-						StaticMesh->RenderData.Reset();
+						StaticMesh->RenderData.Reset(nullptr);
 					}
 
 					if(FDatasmithObjectTemplateUtils::GetObjectTemplate<UDatasmithStaticMeshTemplate>( StaticMesh ) == nullptr)
diff --git a/Engine/Plugins/Enterprise/DatasmithImporter/Source/DatasmithImporter/Private/DatasmithFileProducer.cpp b/Engine/Plugins/Enterprise/DatasmithImporter/Source/DatasmithImporter/Private/DatasmithFileProducer.cpp
--- a/Engine/Plugins/Enterprise/DatasmithImporter/Source/DatasmithImporter/Private/DatasmithFileProducer.cpp
+++ b/Engine/Plugins/Enterprise/DatasmithImporter/Source/DatasmithImporter/Private/DatasmithFileProducer.cpp
@@ -324,9 +324,9 @@ void UDatasmithFileProducer::SceneElementToWorld()
 
 			if (UMaterial* SourceMaterial = Cast< UMaterial >(AssetPair.Value))
 			{
-				SourceMaterial->RebuildExpressionTextureReferences();
+				SourceMaterial->UpdateCachedExpressionData();
 
-				for (FMaterialFunctionInfo& MaterialFunctionInfo : SourceMaterial->MaterialFunctionInfos)
+				for (const FMaterialFunctionInfo& MaterialFunctionInfo : SourceMaterial->GetCachedExpressionData().FunctionInfos)
 				{
 					if (MaterialFunctionInfo.Function && MaterialFunctionInfo.Function->GetOutermost() == SourceMaterial->GetOutermost())
 					{
@@ -471,7 +471,7 @@ void UDatasmithFileProducer::PreventNameCollision()
 				{
 					if (UMaterial* SourceMaterial = Cast< UMaterial >(MaterialInstance))
 					{
-						for (FMaterialFunctionInfo& MaterialFunctionInfo : SourceMaterial->MaterialFunctionInfos)
+						for (const FMaterialFunctionInfo& MaterialFunctionInfo : SourceMaterial->GetCachedExpressionData().FunctionInfos)
 						{
 							if (MaterialFunctionInfo.Function && MaterialFunctionInfo.Function->GetOutermost() == SourceMaterial->GetOutermost())
 							{
diff --git a/Engine/Plugins/Enterprise/DatasmithImporter/Source/DatasmithImporter/Private/DatasmithImporter.cpp b/Engine/Plugins/Enterprise/DatasmithImporter/Source/DatasmithImporter/Private/DatasmithImporter.cpp
--- a/Engine/Plugins/Enterprise/DatasmithImporter/Source/DatasmithImporter/Private/DatasmithImporter.cpp
+++ b/Engine/Plugins/Enterprise/DatasmithImporter/Source/DatasmithImporter/Private/DatasmithImporter.cpp
@@ -830,9 +830,7 @@ namespace DatasmithImporterImpl
 					TextureParametersToConvertMap.Add(BaseMaterial, FMaterialParameterInfo());
 
 					//If no unsupported virtual texture parameters were found, it's possible that a texture needing conversion is simply not exposed as a parameter, so we still need to check for those.
-					TArray<UObject*> ReferencedTextures;
-					BaseMaterial->AppendReferencedTextures(ReferencedTextures);
-					for (UObject* ReferencedTexture : ReferencedTextures)
+					for (UObject* ReferencedTexture : BaseMaterial->GetCachedExpressionData().ReferencedTextures)
 					{
 						if (VirtualTexturesToConvert.Contains(Cast<UTexture2D>(ReferencedTexture)))
 						{
@@ -2389,10 +2387,10 @@ void FDatasmithImporter::FinalizeImport(FDatasmithImportContext& ImportContext,
 
 		if (UMaterial* SourceMaterial = Cast< UMaterial >(SourceMaterialInterface))
 		{
-			SourceMaterial->RebuildExpressionTextureReferences();
-			
+			SourceMaterial->UpdateCachedExpressionData();
+
 			TArray<UObject*> ReferencedTextures;
-			SourceMaterial->AppendReferencedTextures(ReferencedTextures);
+			ReferencedTextures = SourceMaterial->GetReferencedTextures();
 			for (UTexture2D* VirtualTexture : ImportContext.AssetsContext.VirtualTexturesToConvert)
 			{
 				if (ReferencedTextures.Contains(VirtualTexture))
@@ -2402,7 +2400,7 @@ void FDatasmithImporter::FinalizeImport(FDatasmithImportContext& ImportContext,
 				}
 			}
 
-			for (FMaterialFunctionInfo& MaterialFunctionInfo : SourceMaterial->MaterialFunctionInfos)
+			for (const FMaterialFunctionInfo& MaterialFunctionInfo : SourceMaterial->GetCachedExpressionData().FunctionInfos)
 			{
 				if (MaterialFunctionInfo.Function && MaterialFunctionInfo.Function->GetOutermost() == SourceMaterial->GetOutermost())
 				{
diff --git a/Engine/Plugins/Enterprise/LidarPointCloud/Source/LidarPointCloudRuntime/Private/Rendering/LidarPointCloudRenderBuffers.cpp b/Engine/Plugins/Enterprise/LidarPointCloud/Source/LidarPointCloudRuntime/Private/Rendering/LidarPointCloudRenderBuffers.cpp
--- a/Engine/Plugins/Enterprise/LidarPointCloud/Source/LidarPointCloudRuntime/Private/Rendering/LidarPointCloudRenderBuffers.cpp
+++ b/Engine/Plugins/Enterprise/LidarPointCloud/Source/LidarPointCloudRuntime/Private/Rendering/LidarPointCloudRenderBuffers.cpp
@@ -441,9 +441,10 @@ FLidarPointCloudVertexFactory::~FLidarPointCloudVertexFactory()
 	Fence.Wait();
 }
 
-bool FLidarPointCloudVertexFactory::ShouldCompilePermutation(EShaderPlatform Platform, const class FMaterial* Material, const class FShaderType* ShaderType)
+bool FLidarPointCloudVertexFactory::ShouldCompilePermutation(const FVertexFactoryShaderPermutationParameters& Parameters)
 {
-	return (IsPCPlatform(Platform) && IsFeatureLevelSupported(Platform, ERHIFeatureLevel::SM5) && Material->GetMaterialDomain() == MD_Surface && Material->IsUsedWithLidarPointCloud()) || Material->IsSpecialEngineMaterial();
+	return (IsPCPlatform(Parameters.Platform) && IsFeatureLevelSupported(Parameters.Platform, ERHIFeatureLevel::SM5) && 
+		Parameters.MaterialParameters.MaterialDomain == MD_Surface && Parameters.MaterialParameters.bIsUsedWithLidarPointCloud) || Parameters.MaterialParameters.bIsSpecialEngineMaterial;
 }
 
 void FLidarPointCloudVertexFactory::InitRHI()
@@ -465,4 +466,4 @@ IMPLEMENT_VERTEX_FACTORY_TYPE(FLidarPointCloudVertexFactory, "/Plugin/LidarPoint
 
 #undef BINDPARAM
 #undef SETPARAM
-#undef SETSRVPARAM
+#undef SETSRVPARAM
\ No newline at end of file
diff --git a/Engine/Plugins/Enterprise/LidarPointCloud/Source/LidarPointCloudRuntime/Private/Rendering/LidarPointCloudRenderBuffers.h b/Engine/Plugins/Enterprise/LidarPointCloud/Source/LidarPointCloudRuntime/Private/Rendering/LidarPointCloudRenderBuffers.h
--- a/Engine/Plugins/Enterprise/LidarPointCloud/Source/LidarPointCloudRuntime/Private/Rendering/LidarPointCloudRenderBuffers.h
+++ b/Engine/Plugins/Enterprise/LidarPointCloud/Source/LidarPointCloudRuntime/Private/Rendering/LidarPointCloudRenderBuffers.h
@@ -7,6 +7,7 @@
 #include "VertexFactory.h"
 #include "LidarPointCloudShared.h"
 #include "HAL/ThreadSafeBool.h"
+#include "GlobalShader.h"
 
 /**
  * Base class for the buffers
@@ -82,7 +83,7 @@ private:
 /**
  * Handles creation of the dynamic Instance Data
  */
-// #refactor: Can we avoid referencing FLidarPointCloudOctreeNode here?
+ // #refactor: Can we avoid referencing FLidarPointCloudOctreeNode here?
 struct FLidarPointCloudTraversalOctreeNode;
 class FLidarPointCloudInstanceBuffer
 {
@@ -169,16 +170,9 @@ struct FLidarPointCloudBatchElementUserData
 class FLidarPointCloudVertexFactoryShaderParameters : public FVertexFactoryShaderParameters
 {
 public:
-	virtual void Bind(const FShaderParameterMap& ParameterMap) override;
-	virtual void Serialize(FArchive& Ar) override
-	{
-		Ar << DataBuffer << IndexDivisor << VDMultiplier << SizeOffset << RootCellSize << bUseLODColoration << SpriteSizeMultiplier << ViewRightVector << ViewUpVector
-			<< BoundsSize << BoundsOffset << ElevationColorBottom << ElevationColorTop << bUseCircle << bUseColorOverride << bUseElevationColor
-			<< Offset << Contrast << Saturation << Gamma << Tint << IntensityInfluence << bUseClassification << ClassificationColors;
-	}
-
-	virtual void GetElementShaderBindings(const class FSceneInterface* Scene, const FSceneView* View, const FMeshMaterialShader* Shader, const EVertexInputStreamType InputStreamType, ERHIFeatureLevel::Type FeatureLevel,
-		const FVertexFactory* VertexFactory, const FMeshBatchElement& BatchElement, class FMeshDrawSingleShaderBindings& ShaderBindings, FVertexInputStreamArray& VertexStreams) const override;
+	void Bind(const FShaderParameterMap& ParameterMap);
+	void GetElementShaderBindings(const class FSceneInterface* Scene, const FSceneView* View, const FMeshMaterialShader* Shader, const EVertexInputStreamType InputStreamType, ERHIFeatureLevel::Type FeatureLevel,
+		const FVertexFactory* VertexFactory, const FMeshBatchElement& BatchElement, class FMeshDrawSingleShaderBindings& ShaderBindings, FVertexInputStreamArray& VertexStreams) const;
 
 private:
 	FShaderResourceParameter DataBuffer;
@@ -217,9 +211,8 @@ class FLidarPointCloudVertexFactory : public FVertexFactory
 public:
 	static FVertexFactoryShaderParameters* ConstructShaderParameters(EShaderFrequency ShaderFrequency);
 
-	static bool ShouldCache(EShaderPlatform Platform, const class FMaterial* Material, const class FShaderType* ShaderType) { return IsFeatureLevelSupported(Platform, ERHIFeatureLevel::SM5); }
-	static bool ShouldCompilePermutation(EShaderPlatform Platform, const class FMaterial* Material, const class FShaderType* ShaderType);
-	static void ModifyCompilationEnvironment(const FVertexFactoryType* Type, EShaderPlatform Platform, const FMaterial* Material, FShaderCompilerEnvironment& OutEnvironment) {}
+	static bool ShouldCompilePermutation(const FVertexFactoryShaderPermutationParameters& Parameters);
+	static void ModifyCompilationEnvironment(const FVertexFactoryShaderPermutationParameters& Parameters, FShaderCompilerEnvironment& OutEnvironment) {}
 
 	FLidarPointCloudVertexFactory() : FVertexFactory(ERHIFeatureLevel::SM5) { }
 	~FLidarPointCloudVertexFactory();
diff --git a/Engine/Plugins/Enterprise/MDLImporter/Source/MDLImporter/Private/MDLImporter.cpp b/Engine/Plugins/Enterprise/MDLImporter/Source/MDLImporter/Private/MDLImporter.cpp
--- a/Engine/Plugins/Enterprise/MDLImporter/Source/MDLImporter/Private/MDLImporter.cpp
+++ b/Engine/Plugins/Enterprise/MDLImporter/Source/MDLImporter/Private/MDLImporter.cpp
@@ -278,11 +278,9 @@ void FMDLImporter::ConvertUnsuportedVirtualTextures() const
 	{
 		if (UMaterial* CurrentMaterial = Cast<UMaterial>(CurrentMaterialInterface))
 		{
-			TArray<UObject*> ReferencedTextures;
-			CurrentMaterial->AppendReferencedTextures(ReferencedTextures);
 			for (UTexture2D* VirtualTexture : VirtualTexturesToConvert)
 			{
-				if (ReferencedTextures.Contains(VirtualTexture))
+				if (CurrentMaterial->GetCachedExpressionData().ReferencedTextures.Contains(VirtualTexture))
 				{
 					MaterialsToRefreshAfterVirtualTextureConversion.Add(CurrentMaterial);
 					break;
diff --git a/Engine/Plugins/Enterprise/VariantManagerContent/Source/VariantManagerContent/Private/Variant.cpp b/Engine/Plugins/Enterprise/VariantManagerContent/Source/VariantManagerContent/Private/Variant.cpp
--- a/Engine/Plugins/Enterprise/VariantManagerContent/Source/VariantManagerContent/Private/Variant.cpp
+++ b/Engine/Plugins/Enterprise/VariantManagerContent/Source/VariantManagerContent/Private/Variant.cpp
@@ -271,7 +271,7 @@ void UVariant::SetThumbnail(UTexture2D* NewThumbnail)
 
 			if (TargetWidth != OriginalWidth || TargetHeight != OriginalHeight)
 			{
-				FImageUtils::ImageResize(OriginalWidth, OriginalHeight, OriginalColors, TargetWidth, TargetHeight, TargetColors, false);
+				FImageUtils::ImageResize(OriginalWidth, OriginalHeight, OriginalColors, TargetWidth, TargetHeight, TArrayView<FColor>(TargetColors), false);
 			}
 			else
 			{
diff --git a/Engine/Plugins/Experimental/AR/PointCloud/Source/PointCloud/Private/PointCloudVertexFactory.cpp b/Engine/Plugins/Experimental/AR/PointCloud/Source/PointCloud/Private/PointCloudVertexFactory.cpp
--- a/Engine/Plugins/Experimental/AR/PointCloud/Source/PointCloud/Private/PointCloudVertexFactory.cpp
+++ b/Engine/Plugins/Experimental/AR/PointCloud/Source/PointCloud/Private/PointCloudVertexFactory.cpp
@@ -17,20 +17,15 @@ IMPLEMENT_GLOBAL_SHADER_PARAMETER_STRUCT(FPointCloudVertexFactoryParameters, "Po
 class FPointCloudVertexFactoryShaderParameters :
 	public FVertexFactoryShaderParameters
 {
+	DECLARE_INLINE_TYPE_LAYOUT(FPointCloudVertexFactoryShaderParameters, NonVirtual);
 public:
-	virtual void Bind(const FShaderParameterMap& ParameterMap) override
+	void Bind(const FShaderParameterMap& ParameterMap)
 	{
 		ColorMask.Bind(ParameterMap, TEXT("ColorMask"));
 		PointSize.Bind(ParameterMap, TEXT("PointSize"));
 	}
 
-	virtual void Serialize(FArchive& Ar) override
-	{
-		Ar << ColorMask;
-		Ar << PointSize;
-	}
-
-	virtual void GetElementShaderBindings(
+	void GetElementShaderBindings(
 		const class FSceneInterface* Scene,
 		const class FSceneView* View,
 		const class FMeshMaterialShader* Shader,
@@ -39,7 +34,7 @@ public:
 		const class FVertexFactory* InVertexFactory,
 		const struct FMeshBatchElement& BatchElement,
 		class FMeshDrawSingleShaderBindings& ShaderBindings,
-		FVertexInputStreamArray& VertexStreams) const override
+		FVertexInputStreamArray& VertexStreams) const
 	{		
 		FPointCloudVertexFactory* VertexFactory = (FPointCloudVertexFactory*)InVertexFactory;
 
@@ -49,11 +44,9 @@ public:
 		ShaderBindings.Add(PointSize, VertexFactory->GetPointSize());
 	}
 
-	virtual uint32 GetSize() const override { return sizeof(*this); }
-
 private:
-	FShaderParameter ColorMask;
-	FShaderParameter PointSize;
+	LAYOUT_FIELD(FShaderParameter, ColorMask);
+	LAYOUT_FIELD(FShaderParameter, PointSize);
 };
 
 /**
@@ -128,15 +121,10 @@ void FPointCloudVertexFactory::ReleaseRHI()
 	FVertexFactory::ReleaseRHI();
 }
 
-bool FPointCloudVertexFactory::ShouldCompilePermutation(EShaderPlatform Platform, const class FMaterial* Material, const class FShaderType* ShaderType)
+bool FPointCloudVertexFactory::ShouldCompilePermutation(const FVertexFactoryShaderPermutationParameters& Parameters)
 {
 	// We exclusively use manual fetch, so we need that supported
-	return RHISupportsManualVertexFetch(Platform);
-}
-
-FVertexFactoryShaderParameters* FPointCloudVertexFactory::ConstructShaderParameters(EShaderFrequency ShaderFrequency)
-{
-	return ShaderFrequency == SF_Vertex ? new FPointCloudVertexFactoryShaderParameters() : nullptr;
+	return RHISupportsManualVertexFetch(Parameters.Platform);
 }
 
 void FPointCloudVertexFactory::SetParameters(const FPointCloudVertexFactoryParameters& InUniformParameters, const uint32 InMask, const float InSize)
@@ -146,4 +134,5 @@ void FPointCloudVertexFactory::SetParameters(const FPointCloudVertexFactoryParam
 	PointSize = InSize;
 }
 
+IMPLEMENT_VERTEX_FACTORY_PARAMETER_TYPE(FPointCloudVertexFactory, SF_Vertex, FPointCloudVertexFactoryShaderParameters);
 IMPLEMENT_VERTEX_FACTORY_TYPE(FPointCloudVertexFactory, "/Engine/Private/PointCloudVertexFactory.ush", true, false, false, false, false);
diff --git a/Engine/Plugins/Experimental/AR/PointCloud/Source/PointCloud/Private/PointCloudVertexFactory.h b/Engine/Plugins/Experimental/AR/PointCloud/Source/PointCloud/Private/PointCloudVertexFactory.h
--- a/Engine/Plugins/Experimental/AR/PointCloud/Source/PointCloud/Private/PointCloudVertexFactory.h
+++ b/Engine/Plugins/Experimental/AR/PointCloud/Source/PointCloud/Private/PointCloudVertexFactory.h
@@ -43,12 +43,7 @@ public:
 	/**
 	 * Should we cache the material's shadertype on this platform with this vertex factory?
 	 */
-	static bool ShouldCompilePermutation(EShaderPlatform Platform, const class FMaterial* Material, const class FShaderType* ShaderType);
-
-	/**
-	 * Construct shader parameters for this type of vertex factory.
-	 */
-	static FVertexFactoryShaderParameters* ConstructShaderParameters(EShaderFrequency ShaderFrequency);
+	static bool ShouldCompilePermutation(const FVertexFactoryShaderPermutationParameters& Parameters);
 
 	/**
 	 * Set parameters for this vertex factory instance.
diff --git a/Engine/Plugins/Experimental/ChaosNiagara/Source/ChaosNiagara/Classes/NiagaraDataInterfaceChaosDestruction.h b/Engine/Plugins/Experimental/ChaosNiagara/Source/ChaosNiagara/Classes/NiagaraDataInterfaceChaosDestruction.h
--- a/Engine/Plugins/Experimental/ChaosNiagara/Source/ChaosNiagara/Classes/NiagaraDataInterfaceChaosDestruction.h
+++ b/Engine/Plugins/Experimental/ChaosNiagara/Source/ChaosNiagara/Classes/NiagaraDataInterfaceChaosDestruction.h
@@ -352,6 +352,9 @@ class CHAOSNIAGARA_API UNiagaraDataInterfaceChaosDestruction : public UNiagaraDa
 	GENERATED_UCLASS_BODY()
 
 public:
+
+	DECLARE_NIAGARA_DI_PARAMETER();
+
 	/* Chaos Solver */
 	UPROPERTY(EditAnywhere, Category = "Solver", meta = (DisplayName = "Chaos Solver"))
 	TSet<AChaosSolverActor*> ChaosSolverActorSet;
@@ -556,7 +559,6 @@ public:
 	// GPU sim functionality
 	virtual void GetParameterDefinitionHLSL(const FNiagaraDataInterfaceGPUParamInfo& ParamInfo, FString& OutHLSL) override;
 	virtual bool GetFunctionHLSL(const FNiagaraDataInterfaceGPUParamInfo& ParamInfo, const FNiagaraDataInterfaceGeneratedFunction& FunctionInfo, int FunctionInstanceIndex, FString& OutHLSL) override;
-	virtual FNiagaraDataInterfaceParametersCS* ConstructComputeParameters() const override;
 
 	//----------------------------------------------------------------------------
 	// EXPOSED FUNCTIONS
diff --git a/Engine/Plugins/Experimental/ChaosNiagara/Source/ChaosNiagara/Private/NiagaraDataInterfaceChaosDestruction.cpp b/Engine/Plugins/Experimental/ChaosNiagara/Source/ChaosNiagara/Private/NiagaraDataInterfaceChaosDestruction.cpp
--- a/Engine/Plugins/Experimental/ChaosNiagara/Source/ChaosNiagara/Private/NiagaraDataInterfaceChaosDestruction.cpp
+++ b/Engine/Plugins/Experimental/ChaosNiagara/Source/ChaosNiagara/Private/NiagaraDataInterfaceChaosDestruction.cpp
@@ -4282,81 +4282,49 @@ void FNiagaraDataInterfaceProxyChaosDestruction::ConsumePerInstanceDataFromGameT
 	}
 }
 
-struct FNiagaraDataInterfaceParametersCS_ChaosDestruction : public FNiagaraDataInterfaceParametersCS
+struct FNiagaraDataInterfaceParametersCS_ChaosDestruction :public FNiagaraDataInterfaceParametersCS
 {
-	virtual void Bind(const FNiagaraDataInterfaceParamRef& ParamRef, const class FShaderParameterMap& ParameterMap) override
-	{
-		PositionBuffer.Bind(ParameterMap, *("PositionBuffer_" + ParamRef.ParameterInfo.DataInterfaceHLSLSymbol));
-		VelocityBuffer.Bind(ParameterMap, *("VelocityBuffer_" + ParamRef.ParameterInfo.DataInterfaceHLSLSymbol));
-		ExtentMinBuffer.Bind(ParameterMap, *("ExtentMinBuffer_" + ParamRef.ParameterInfo.DataInterfaceHLSLSymbol));
-		ExtentMaxBuffer.Bind(ParameterMap, *("ExtentMaxBuffer_" + ParamRef.ParameterInfo.DataInterfaceHLSLSymbol));
-		VolumeBuffer.Bind(ParameterMap, *("VolumeBuffer_" + ParamRef.ParameterInfo.DataInterfaceHLSLSymbol));
-		SolverIDBuffer.Bind(ParameterMap, *("SolverIDBuffer_" + ParamRef.ParameterInfo.DataInterfaceHLSLSymbol));
-		DensityBuffer.Bind(ParameterMap, *("DensityBuffer_" + ParamRef.ParameterInfo.DataInterfaceHLSLSymbol));
-		FrictionBuffer.Bind(ParameterMap, *("FrictionBuffer_" + ParamRef.ParameterInfo.DataInterfaceHLSLSymbol));
-		RestitutionBuffer.Bind(ParameterMap, *("RestitutionBuffer_" + ParamRef.ParameterInfo.DataInterfaceHLSLSymbol));
-		SurfaceTypeBuffer.Bind(ParameterMap, *("SurfaceTypeBuffer_" + ParamRef.ParameterInfo.DataInterfaceHLSLSymbol));
-		ColorBuffer.Bind(ParameterMap, *("ColorBuffer_" + ParamRef.ParameterInfo.DataInterfaceHLSLSymbol));
-
-		IncomingLocationBuffer.Bind(ParameterMap, *("IncomingLocationBuffer_" + ParamRef.ParameterInfo.DataInterfaceHLSLSymbol));
-		IncomingAccumulatedImpulseBuffer.Bind(ParameterMap, *("IncomingAccumulatedImpulseBuffer_" + ParamRef.ParameterInfo.DataInterfaceHLSLSymbol));
-		IncomingNormalBuffer.Bind(ParameterMap, *("IncomingNormalBuffer_" + ParamRef.ParameterInfo.DataInterfaceHLSLSymbol));
-		IncomingVelocity1Buffer.Bind(ParameterMap, *("IncomingVelocity1Buffer_" + ParamRef.ParameterInfo.DataInterfaceHLSLSymbol));
-		IncomingVelocity2Buffer.Bind(ParameterMap, *("IncomingVelocity2Buffer_" + ParamRef.ParameterInfo.DataInterfaceHLSLSymbol));
-		IncomingAngularVelocity1Buffer.Bind(ParameterMap, *("IncomingAngularVelocity1Buffer_" + ParamRef.ParameterInfo.DataInterfaceHLSLSymbol));
-		IncomingAngularVelocity2Buffer.Bind(ParameterMap, *("IncomingAngularVelocity2Buffer_" + ParamRef.ParameterInfo.DataInterfaceHLSLSymbol));
-		IncomingMass1Buffer.Bind(ParameterMap, *("IncomingMass1Buffer_" + ParamRef.ParameterInfo.DataInterfaceHLSLSymbol));
-		IncomingMass2Buffer.Bind(ParameterMap, *("IncomingMass2Buffer_" + ParamRef.ParameterInfo.DataInterfaceHLSLSymbol));
-		IncomingTimeBuffer.Bind(ParameterMap, *("IncomingTimeBuffer_" + ParamRef.ParameterInfo.DataInterfaceHLSLSymbol));
-
-		TransformTranslationBuffer.Bind(ParameterMap, *("TransformTranslationBuffer_" + ParamRef.ParameterInfo.DataInterfaceHLSLSymbol));
-		TransformRotationBuffer.Bind(ParameterMap, *("TransformRotationBuffer_" + ParamRef.ParameterInfo.DataInterfaceHLSLSymbol));
-		TransformScaleBuffer.Bind(ParameterMap, *("TransformScaleBuffer_" + ParamRef.ParameterInfo.DataInterfaceHLSLSymbol));
-		BoundsBuffer.Bind(ParameterMap, *("BoundsBuffer_" + ParamRef.ParameterInfo.DataInterfaceHLSLSymbol));
-
-		LastSpawnedPointID.Bind(ParameterMap, *("LastSpawnedPointID_" + ParamRef.ParameterInfo.DataInterfaceHLSLSymbol));
-		SolverTime.Bind(ParameterMap, *("SolverTime_" + ParamRef.ParameterInfo.DataInterfaceHLSLSymbol));
-	}
-
-	virtual void Serialize(FArchive& Ar)override
-	{
-		Ar << PositionBuffer;
-		Ar << VelocityBuffer;
-		Ar << ExtentMinBuffer;
-		Ar << ExtentMaxBuffer;
-		Ar << VolumeBuffer;
-		Ar << SolverIDBuffer;
-		Ar << DensityBuffer;
-		Ar << FrictionBuffer;
-		Ar << RestitutionBuffer;
-		Ar << SurfaceTypeBuffer;
-		Ar << ColorBuffer;
-
-		Ar << IncomingLocationBuffer;
-		Ar << IncomingAccumulatedImpulseBuffer;
-		Ar << IncomingNormalBuffer;
-		Ar << IncomingVelocity1Buffer;
-		Ar << IncomingVelocity2Buffer;
-		Ar << IncomingAngularVelocity1Buffer;
-		Ar << IncomingAngularVelocity2Buffer;
-		Ar << IncomingMass1Buffer;
-		Ar << IncomingMass2Buffer;
-		Ar << IncomingTimeBuffer;
-
-		Ar << TransformTranslationBuffer;
-		Ar << TransformRotationBuffer;
-		Ar << TransformScaleBuffer;
-		Ar << BoundsBuffer;
-
-		Ar << LastSpawnedPointID;
-		Ar << SolverTime;
-	}
-
-	virtual void Set(FRHICommandList& RHICmdList, const FNiagaraDataInterfaceSetArgs& Context) const override
+	DECLARE_TYPE_LAYOUT(FNiagaraDataInterfaceParametersCS_ChaosDestruction, NonVirtual);
+public:
+	void Bind(const FNiagaraDataInterfaceGPUParamInfo& ParameterInfo, const class FShaderParameterMap& ParameterMap)
+	{
+		PositionBuffer.Bind(ParameterMap, *("PositionBuffer_" + ParameterInfo.DataInterfaceHLSLSymbol));
+		VelocityBuffer.Bind(ParameterMap, *("VelocityBuffer_" + ParameterInfo.DataInterfaceHLSLSymbol));
+		ExtentMinBuffer.Bind(ParameterMap, *("ExtentMinBuffer_" + ParameterInfo.DataInterfaceHLSLSymbol));
+		ExtentMaxBuffer.Bind(ParameterMap, *("ExtentMaxBuffer_" + ParameterInfo.DataInterfaceHLSLSymbol));
+		VolumeBuffer.Bind(ParameterMap, *("VolumeBuffer_" + ParameterInfo.DataInterfaceHLSLSymbol));
+		SolverIDBuffer.Bind(ParameterMap, *("SolverIDBuffer_" + ParameterInfo.DataInterfaceHLSLSymbol));
+		DensityBuffer.Bind(ParameterMap, *("DensityBuffer_" + ParameterInfo.DataInterfaceHLSLSymbol));
+		FrictionBuffer.Bind(ParameterMap, *("FrictionBuffer_" + ParameterInfo.DataInterfaceHLSLSymbol));
+		RestitutionBuffer.Bind(ParameterMap, *("RestitutionBuffer_" + ParameterInfo.DataInterfaceHLSLSymbol));
+		SurfaceTypeBuffer.Bind(ParameterMap, *("SurfaceTypeBuffer_" + ParameterInfo.DataInterfaceHLSLSymbol));
+		ColorBuffer.Bind(ParameterMap, *("ColorBuffer_" + ParameterInfo.DataInterfaceHLSLSymbol));
+
+		IncomingLocationBuffer.Bind(ParameterMap, *("IncomingLocationBuffer_" + ParameterInfo.DataInterfaceHLSLSymbol));
+		IncomingAccumulatedImpulseBuffer.Bind(ParameterMap, *("IncomingAccumulatedImpulseBuffer_" + ParameterInfo.DataInterfaceHLSLSymbol));
+		IncomingNormalBuffer.Bind(ParameterMap, *("IncomingNormalBuffer_" + ParameterInfo.DataInterfaceHLSLSymbol));
+		IncomingVelocity1Buffer.Bind(ParameterMap, *("IncomingVelocity1Buffer_" + ParameterInfo.DataInterfaceHLSLSymbol));
+		IncomingVelocity2Buffer.Bind(ParameterMap, *("IncomingVelocity2Buffer_" + ParameterInfo.DataInterfaceHLSLSymbol));
+		IncomingAngularVelocity1Buffer.Bind(ParameterMap, *("IncomingAngularVelocity1Buffer_" + ParameterInfo.DataInterfaceHLSLSymbol));
+		IncomingAngularVelocity2Buffer.Bind(ParameterMap, *("IncomingAngularVelocity2Buffer_" + ParameterInfo.DataInterfaceHLSLSymbol));
+		IncomingMass1Buffer.Bind(ParameterMap, *("IncomingMass1Buffer_" + ParameterInfo.DataInterfaceHLSLSymbol));
+		IncomingMass2Buffer.Bind(ParameterMap, *("IncomingMass2Buffer_" + ParameterInfo.DataInterfaceHLSLSymbol));
+		IncomingTimeBuffer.Bind(ParameterMap, *("IncomingTimeBuffer_" + ParameterInfo.DataInterfaceHLSLSymbol));
+
+		TransformTranslationBuffer.Bind(ParameterMap, *("TransformTranslationBuffer_" + ParameterInfo.DataInterfaceHLSLSymbol));
+		TransformRotationBuffer.Bind(ParameterMap, *("TransformRotationBuffer_" + ParameterInfo.DataInterfaceHLSLSymbol));
+		TransformScaleBuffer.Bind(ParameterMap, *("TransformScaleBuffer_" + ParameterInfo.DataInterfaceHLSLSymbol));
+		BoundsBuffer.Bind(ParameterMap, *("BoundsBuffer_" + ParameterInfo.DataInterfaceHLSLSymbol));
+
+		LastSpawnedPointID.Bind(ParameterMap, *("LastSpawnedPointID_" + ParameterInfo.DataInterfaceHLSLSymbol));
+		SolverTime.Bind(ParameterMap, *("SolverTime_" + ParameterInfo.DataInterfaceHLSLSymbol));
+	}
+
+	void Set(FRHICommandList& RHICmdList, const FNiagaraDataInterfaceSetArgs& Context) const
 	{
 		check(IsInRenderingThread());
 
-		FRHIComputeShader* ComputeShaderRHI = Context.Shader->GetComputeShader();
+		FRHIComputeShader* ComputeShaderRHI = Context.Shader.GetComputeShader();
 		FNiagaraDataInterfaceProxyChaosDestruction* ChaosDestructionInterfaceProxy = static_cast<FNiagaraDataInterfaceProxyChaosDestruction*>(Context.DataInterface);
 		if (ChaosDestructionInterfaceProxy)
 		{
@@ -4406,42 +4374,43 @@ struct FNiagaraDataInterfaceParametersCS_ChaosDestruction : public FNiagaraDataI
 
 private:
 	// TODO: Collect these into a small number of buffers to reduce the number of binding points
-	FShaderResourceParameter PositionBuffer;
-	FShaderResourceParameter VelocityBuffer;
-	FShaderResourceParameter ExtentMinBuffer;
-	FShaderResourceParameter ExtentMaxBuffer;
-	FShaderResourceParameter VolumeBuffer;
-	FShaderResourceParameter SolverIDBuffer;
-	FShaderResourceParameter DensityBuffer;
-	FShaderResourceParameter FrictionBuffer;
-	FShaderResourceParameter RestitutionBuffer;
-	FShaderResourceParameter SurfaceTypeBuffer;
-	FShaderResourceParameter ColorBuffer;
-
-	FShaderResourceParameter IncomingLocationBuffer;
-	FShaderResourceParameter IncomingAccumulatedImpulseBuffer;
-	FShaderResourceParameter IncomingNormalBuffer;
-	FShaderResourceParameter IncomingVelocity1Buffer;
-	FShaderResourceParameter IncomingVelocity2Buffer;
-	FShaderResourceParameter IncomingAngularVelocity1Buffer;
-	FShaderResourceParameter IncomingAngularVelocity2Buffer;
-	FShaderResourceParameter IncomingMass1Buffer;
-	FShaderResourceParameter IncomingMass2Buffer;
-	FShaderResourceParameter IncomingTimeBuffer;
-
-	FShaderResourceParameter TransformTranslationBuffer;
-	FShaderResourceParameter TransformRotationBuffer;
-	FShaderResourceParameter TransformScaleBuffer;
-	FShaderResourceParameter BoundsBuffer;
-
-	FShaderParameter LastSpawnedPointID;
-	FShaderParameter SolverTime;
+	LAYOUT_FIELD(FShaderResourceParameter, PositionBuffer);
+	LAYOUT_FIELD(FShaderResourceParameter, VelocityBuffer);
+	LAYOUT_FIELD(FShaderResourceParameter, ExtentMinBuffer);
+	LAYOUT_FIELD(FShaderResourceParameter, ExtentMaxBuffer);
+	LAYOUT_FIELD(FShaderResourceParameter, VolumeBuffer);
+	LAYOUT_FIELD(FShaderResourceParameter, SolverIDBuffer);
+	LAYOUT_FIELD(FShaderResourceParameter, DensityBuffer);
+	LAYOUT_FIELD(FShaderResourceParameter, FrictionBuffer);
+	LAYOUT_FIELD(FShaderResourceParameter, RestitutionBuffer);
+	LAYOUT_FIELD(FShaderResourceParameter, SurfaceTypeBuffer);
+	LAYOUT_FIELD(FShaderResourceParameter, ColorBuffer);
+
+	LAYOUT_FIELD(FShaderResourceParameter, IncomingLocationBuffer);
+	LAYOUT_FIELD(FShaderResourceParameter, IncomingAccumulatedImpulseBuffer);
+	LAYOUT_FIELD(FShaderResourceParameter, IncomingNormalBuffer);
+	LAYOUT_FIELD(FShaderResourceParameter, IncomingVelocity1Buffer);
+	LAYOUT_FIELD(FShaderResourceParameter, IncomingVelocity2Buffer);
+	LAYOUT_FIELD(FShaderResourceParameter, IncomingAngularVelocity1Buffer);
+	LAYOUT_FIELD(FShaderResourceParameter, IncomingAngularVelocity2Buffer);
+	LAYOUT_FIELD(FShaderResourceParameter, IncomingMass1Buffer);
+	LAYOUT_FIELD(FShaderResourceParameter, IncomingMass2Buffer);
+	LAYOUT_FIELD(FShaderResourceParameter, IncomingTimeBuffer);
+
+	LAYOUT_FIELD(FShaderResourceParameter, TransformTranslationBuffer);
+	LAYOUT_FIELD(FShaderResourceParameter, TransformRotationBuffer);
+	LAYOUT_FIELD(FShaderResourceParameter, TransformScaleBuffer);
+	LAYOUT_FIELD(FShaderResourceParameter, BoundsBuffer);
+
+	LAYOUT_FIELD(FShaderParameter, LastSpawnedPointID);
+	LAYOUT_FIELD(FShaderParameter, SolverTime);
 };
 
-FNiagaraDataInterfaceParametersCS* UNiagaraDataInterfaceChaosDestruction::ConstructComputeParameters()const
-{
-	return new FNiagaraDataInterfaceParametersCS_ChaosDestruction();
-}
+IMPLEMENT_TYPE_LAYOUT(FNiagaraDataInterfaceParametersCS_ChaosDestruction);
+
+IMPLEMENT_NIAGARA_DI_PARAMETER(UNiagaraDataInterfaceChaosDestruction, FNiagaraDataInterfaceParametersCS_ChaosDestruction);
+
+
 //#pragma optimize("", on)
 
 
diff --git a/Engine/Plugins/Experimental/GeometryCache/Source/GeometryCache/Private/GeometryCacheSceneProxy.cpp b/Engine/Plugins/Experimental/GeometryCache/Source/GeometryCache/Private/GeometryCacheSceneProxy.cpp
--- a/Engine/Plugins/Experimental/GeometryCache/Source/GeometryCache/Private/GeometryCacheSceneProxy.cpp
+++ b/Engine/Plugins/Experimental/GeometryCache/Source/GeometryCache/Private/GeometryCacheSceneProxy.cpp
@@ -577,7 +577,7 @@ void FGeometryCacheSceneProxy::UpdateAnimation(float NewTime, bool bNewLooping,
 				Section->RayTracingGeometry.Initializer.IndexBuffer = Section->IndexBuffer.IndexBufferRHI;
 				Section->RayTracingGeometry.Initializer.TotalPrimitiveCount = 0;
 				
-				TArray<FRayTracingGeometrySegment>& Segments = Section->RayTracingGeometry.Initializer.Segments;
+				TMemoryImageArray<FRayTracingGeometrySegment>& Segments = Section->RayTracingGeometry.Initializer.Segments;
 				Segments.Reset();
 
 				for (FGeometryCacheMeshBatchInfo& BatchInfo : Section->MeshData->BatchesInfo)
diff --git a/Engine/Plugins/Experimental/GeometryCollectionPlugin/Source/GeometryCollectionEditor/Private/GeometryCollection/GeometryCollectionConversion.cpp b/Engine/Plugins/Experimental/GeometryCollectionPlugin/Source/GeometryCollectionEditor/Private/GeometryCollection/GeometryCollectionConversion.cpp
--- a/Engine/Plugins/Experimental/GeometryCollectionPlugin/Source/GeometryCollectionEditor/Private/GeometryCollection/GeometryCollectionConversion.cpp
+++ b/Engine/Plugins/Experimental/GeometryCollectionPlugin/Source/GeometryCollectionEditor/Private/GeometryCollection/GeometryCollectionConversion.cpp
@@ -203,8 +203,7 @@ void FGeometryCollectionConversion::AppendStaticMesh(const UStaticMesh * StaticM
 		// necessary since we reindex after all the meshes are added, but it is a good step to have
 		// optimal min/max vertex index right from the static mesh.  All we really need to do is
 		// assign material ids and rely on reindexing, in theory
-		const TArray<FStaticMeshSection> &StaticMeshSections = StaticMesh->RenderData->LODResources[0].Sections;
-		for (const FStaticMeshSection &CurrSection : StaticMeshSections)
+		for (const FStaticMeshSection &CurrSection : StaticMesh->RenderData->LODResources[0].Sections)
 		{			
 			// create new section
 			int32 SectionIndex = GeometryCollection->AddElements(1, FGeometryCollection::MaterialGroup);
diff --git a/Engine/Plugins/Experimental/HairStrands/Source/HairStrandsCore/Private/GroomComponent.cpp b/Engine/Plugins/Experimental/HairStrands/Source/HairStrandsCore/Private/GroomComponent.cpp
--- a/Engine/Plugins/Experimental/HairStrands/Source/HairStrandsCore/Private/GroomComponent.cpp
+++ b/Engine/Plugins/Experimental/HairStrands/Source/HairStrandsCore/Private/GroomComponent.cpp
@@ -56,17 +56,17 @@ public:
 		HairClipLengthParamName(NAME_BoolProperty)
 	{}
 
-	virtual const FMaterial& GetMaterialWithFallback(ERHIFeatureLevel::Type InFeatureLevel, const FMaterialRenderProxy*& OutFallbackMaterialRenderProxy) const
+	virtual const FMaterial& GetMaterialWithFallback(ERHIFeatureLevel::Type InFeatureLevel, const FMaterialRenderProxy*& OutFallbackMaterialRenderProxy) const override
 	{
 		return Parent->GetMaterialWithFallback(InFeatureLevel, OutFallbackMaterialRenderProxy);
 	}
 
-	virtual bool GetVectorValue(const FMaterialParameterInfo& ParameterInfo, FLinearColor* OutValue, const FMaterialRenderContext& Context) const
+	virtual bool GetVectorValue(const FHashedMaterialParameterInfo& ParameterInfo, FLinearColor* OutValue, const FMaterialRenderContext& Context) const override
 	{
 		return Parent->GetVectorValue(ParameterInfo, OutValue, Context);
 	}
 
-	virtual bool GetScalarValue(const FMaterialParameterInfo& ParameterInfo, float* OutValue, const FMaterialRenderContext& Context) const
+	virtual bool GetScalarValue(const FHashedMaterialParameterInfo& ParameterInfo, float* OutValue, const FMaterialRenderContext& Context) const override
 	{
 		if (ParameterInfo.Name == DebugModeParamName)
 		{
@@ -94,12 +94,12 @@ public:
 		}
 	}
 
-	virtual bool GetTextureValue(const FMaterialParameterInfo& ParameterInfo, const URuntimeVirtualTexture** OutValue, const FMaterialRenderContext& Context) const
+	virtual bool GetTextureValue(const FHashedMaterialParameterInfo& ParameterInfo, const URuntimeVirtualTexture** OutValue, const FMaterialRenderContext& Context) const override
 	{
 		return Parent->GetTextureValue(ParameterInfo, OutValue, Context);
 	}
 
-	virtual bool GetTextureValue(const FMaterialParameterInfo& ParameterInfo, const UTexture** OutValue, const FMaterialRenderContext& Context) const
+	virtual bool GetTextureValue(const FHashedMaterialParameterInfo& ParameterInfo, const UTexture** OutValue, const FMaterialRenderContext& Context) const override
 	{
 		return Parent->GetTextureValue(ParameterInfo, OutValue, Context);
 	}
@@ -1724,7 +1724,7 @@ FGroomComponentRecreateRenderStateContext::~FGroomComponentRecreateRenderStateCo
 		if (GroomComponent->IsRegistered() && !GroomComponent->IsRenderStateCreated())
 		{
 			GroomComponent->InitResources();
-			GroomComponent->CreateRenderState_Concurrent();
+			GroomComponent->CreateRenderState_Concurrent(nullptr);
 		}
 	}
 }
diff --git a/Engine/Plugins/Experimental/HairStrands/Source/HairStrandsCore/Private/HairStrandsRendering.cpp b/Engine/Plugins/Experimental/HairStrands/Source/HairStrandsCore/Private/HairStrandsRendering.cpp
--- a/Engine/Plugins/Experimental/HairStrands/Source/HairStrandsCore/Private/HairStrandsRendering.cpp
+++ b/Engine/Plugins/Experimental/HairStrands/Source/HairStrandsCore/Private/HairStrandsRendering.cpp
@@ -197,13 +197,13 @@ static void AddDeformSimHairStrandsPass(
 	PermutationVector.Set<FDeformGuideCS::FGroupSize>(GetGroupSizePermutation(GroupSize));
 	PermutationVector.Set<FDeformGuideCS::FDeformationType>(InternalDeformationType);
 
-	TShaderMap<FGlobalShaderType>* ShaderMap = GetGlobalShaderMap(ERHIFeatureLevel::SM5);
+	FGlobalShaderMap* ShaderMap = GetGlobalShaderMap(ERHIFeatureLevel::SM5);
 
 	TShaderMapRef<FDeformGuideCS> ComputeShader(ShaderMap, PermutationVector);
 	FComputeShaderUtils::AddPass(
 		GraphBuilder,
 		RDG_EVENT_NAME("DeformSimHairStrands"),
-		*ComputeShader,
+		ComputeShader,
 		Parameters,
 		FIntVector(DispatchCountX, DispatchCountY, 1));
 }
@@ -440,13 +440,13 @@ static void AddHairStrandsInterpolationPass(
 	PermutationVector.Set<FHairInterpolationCS::FSimulation>(bHasSimulationEnable ? 1 : 0);
 	PermutationVector.Set<FHairInterpolationCS::FScaleAndClip>(bNeedScaleOrClip ? 1 : 0);
 
-	TShaderMap<FGlobalShaderType>* ShaderMap = GetGlobalShaderMap(ERHIFeatureLevel::SM5);
+	FGlobalShaderMap* ShaderMap = GetGlobalShaderMap(ERHIFeatureLevel::SM5);
 
 	TShaderMapRef<FHairInterpolationCS> ComputeShader(ShaderMap, PermutationVector);
 	FComputeShaderUtils::AddPass(
 		GraphBuilder,
 		RDG_EVENT_NAME("HairStrandsInterpolation"),
-		*ComputeShader,
+		ComputeShader,
 		Parameters,
 		DispatchCount);
 }
@@ -508,7 +508,7 @@ static void AddHairClusterAABBPass(
 	FComputeShaderUtils::AddPass(
 		GraphBuilder,
 		RDG_EVENT_NAME("HairStrandsClusterAABB"),
-		*ComputeShader,
+		ComputeShader,
 		Parameters,
 		DispatchCount);
 }
@@ -554,13 +554,13 @@ static void AddHairTangentPass(
 	FHairTangentCS::FPermutationDomain PermutationVector;
 	PermutationVector.Set<FHairTangentCS::FGroupSize>(GetGroupSizePermutation(GroupSize));
 
-	TShaderMap<FGlobalShaderType>* ShaderMap = GetGlobalShaderMap(ERHIFeatureLevel::SM5);
+	FGlobalShaderMap* ShaderMap = GetGlobalShaderMap(ERHIFeatureLevel::SM5);
 
 	TShaderMapRef<FHairTangentCS> ComputeShader(ShaderMap, PermutationVector);
 	FComputeShaderUtils::AddPass(
 		GraphBuilder,
 		RDG_EVENT_NAME("HairStrandsTangent"),
-		*ComputeShader,
+		ComputeShader,
 		Parameters,
 		DispatchCount);
 }
@@ -612,13 +612,13 @@ static void AddGenerateRaytracingGeometryPass(
 	FHairRaytracingGeometryCS::FPermutationDomain PermutationVector;
 	PermutationVector.Set<FHairRaytracingGeometryCS::FGroupSize>(GetGroupSizePermutation(GroupSize));
 
-	TShaderMap<FGlobalShaderType>* ShaderMap = GetGlobalShaderMap(ERHIFeatureLevel::SM5);
+	FGlobalShaderMap* ShaderMap = GetGlobalShaderMap(ERHIFeatureLevel::SM5);
 
 	TShaderMapRef<FHairRaytracingGeometryCS> ComputeShader(ShaderMap, PermutationVector);
 	FComputeShaderUtils::AddPass(
 		GraphBuilder,
 		RDG_EVENT_NAME("HairStrandsRaytracingGeometry"),
-		*ComputeShader,
+		ComputeShader,
 		Parameters,
 		DispatchCount);
 }
@@ -661,14 +661,14 @@ static void AddClearClusterAABBPass(
 	Parameters->OutGroupAABBBuffer = OutGroupAABBuffer;
 
 	
-	TShaderMap<FGlobalShaderType>* ShaderMap = GetGlobalShaderMap(ERHIFeatureLevel::SM5);
+	FGlobalShaderMap* ShaderMap = GetGlobalShaderMap(ERHIFeatureLevel::SM5);
 	TShaderMapRef<FClearClusterAABBCS> ComputeShader(ShaderMap);
 
 	const FIntVector DispatchCount = DispatchCount.DivideAndRoundUp(FIntVector(ClusterCount * 6, 1, 1) , FIntVector(64, 1, 1));
 	FComputeShaderUtils::AddPass(
 		GraphBuilder,
 		RDG_EVENT_NAME("HairStrandsClearClusterAABB"),
-		*ComputeShader,
+		ComputeShader,
 		Parameters,
 		DispatchCount);
 }
