Force render target back to readable after pass

diff --git a/Engine/Plugins/Compositing/OpenColorIO/Source/OpenColorIO/Private/OpenColorIORendering.cpp b/Engine/Plugins/Compositing/OpenColorIO/Source/OpenColorIO/Private/OpenColorIORendering.cpp
--- a/Engine/Plugins/Compositing/OpenColorIO/Source/OpenColorIO/Private/OpenColorIORendering.cpp
+++ b/Engine/Plugins/Compositing/OpenColorIO/Source/OpenColorIO/Private/OpenColorIORendering.cpp
@@ -105,6 +105,9 @@ static void ProcessOCIOColorSpaceTransform_RenderThread(
 
 	// Resolve render target.
 	InRHICmdList.EndRenderPass();
+
+	// Restore readable state
+	InRHICmdList.TransitionResource(EResourceTransitionAccess::EReadable, OutputSpaceColorResource->TextureRHI);
 }
 
 // static
