Fix loading OpenColorIO SMs for projects with >1 RHI (UE-91197).

- E.g. if a project targets both Vulkan and D3D, D3D shadermaps need to be discarded on load (if running with -vulkan).

diff --git a/Engine/Plugins/Compositing/OpenColorIO/Source/OpenColorIO/Private/OpenColorIOShaderMap.cpp b/Engine/Plugins/Compositing/OpenColorIO/Source/OpenColorIO/Private/OpenColorIOShaderMap.cpp
--- a/Engine/Plugins/Compositing/OpenColorIO/Source/OpenColorIO/Private/OpenColorIOShaderMap.cpp
+++ b/Engine/Plugins/Compositing/OpenColorIO/Source/OpenColorIO/Private/OpenColorIOShaderMap.cpp
@@ -703,12 +703,12 @@ void FOpenColorIOShaderMap::FlushShadersByShaderType(const FShaderType* InShader
 }
 
 
-void FOpenColorIOShaderMap::Serialize(FArchive& Ar, bool bInlineShaderResources)
+bool FOpenColorIOShaderMap::Serialize(FArchive& Ar, bool bInlineShaderResources)
 {
 	// Note: This is saved to the DDC, not into packages (except when cooked)
 	// Backwards compatibility therefore will not work based on the version of Ar
 	// Instead, just bump OPENCOLORIO_DERIVEDDATA_VER
-	Super::Serialize(Ar, bInlineShaderResources, false);
+	return Super::Serialize(Ar, bInlineShaderResources, false);
 }
 
 void FOpenColorIOShaderMap::RemovePendingColorTransform(FOpenColorIOTransformResource* InColorTransform)
diff --git a/Engine/Plugins/Compositing/OpenColorIO/Source/OpenColorIO/Private/OpenColorIOShared.cpp b/Engine/Plugins/Compositing/OpenColorIO/Source/OpenColorIO/Private/OpenColorIOShared.cpp
--- a/Engine/Plugins/Compositing/OpenColorIO/Source/OpenColorIO/Private/OpenColorIOShared.cpp
+++ b/Engine/Plugins/Compositing/OpenColorIO/Source/OpenColorIO/Private/OpenColorIOShared.cpp
@@ -169,11 +169,10 @@ void FOpenColorIOTransformResource::SerializeShaderMap(FArchive& Ar)
 			if (bValid)
 			{
 				TRefCountPtr<FOpenColorIOShaderMap> LoadedShaderMap = new FOpenColorIOShaderMap();
-				LoadedShaderMap->Serialize(Ar);
+				bool bSuccessfullyLoaded = LoadedShaderMap->Serialize(Ar);
 
-				// Toss the loaded shader data if this is a server only instance
-				//@todo - don't cook it in the first place
-				if (FApp::CanEverRender())
+				// Toss the loaded shader data if this is a server only instance (@todo - don't cook it in the first place) or if it's for a different RHI than the current one
+				if (bSuccessfullyLoaded && FApp::CanEverRender())
 				{
 					GameThreadShaderMap = RenderingThreadShaderMap = LoadedShaderMap;
 				}
diff --git a/Engine/Plugins/Compositing/OpenColorIO/Source/OpenColorIO/Public/OpenColorIOShared.h b/Engine/Plugins/Compositing/OpenColorIO/Source/OpenColorIO/Public/OpenColorIOShared.h
--- a/Engine/Plugins/Compositing/OpenColorIO/Source/OpenColorIO/Public/OpenColorIOShared.h
+++ b/Engine/Plugins/Compositing/OpenColorIO/Source/OpenColorIO/Public/OpenColorIOShared.h
@@ -208,7 +208,7 @@ public:
 	static const FOpenColorIOShaderMap* GetShaderMapBeingCompiled(const FOpenColorIOTransformResource* InColorTransform);
 
 	/** Serializes the shader map. */
-	void Serialize(FArchive& Ar, bool bInlineShaderResources = true);
+	bool Serialize(FArchive& Ar, bool bInlineShaderResources = true);
 
 	/** Saves this shader map to the derived data cache. */
 	void SaveToDerivedDataCache();
