Add freezable hash table, use to accelerate lookup of shaders within ShaderMap
- Shaders are searched by TypeName and PermutationId, so move those fields out of FShader and into separate arrays inside ShaderMap, to reduce cache misses while finding shaders

diff --git a/Engine/Plugins/Compositing/OpenColorIO/Source/OpenColorIO/Private/OpenColorIODerivedDataVersion.h b/Engine/Plugins/Compositing/OpenColorIO/Source/OpenColorIO/Private/OpenColorIODerivedDataVersion.h
--- a/Engine/Plugins/Compositing/OpenColorIO/Source/OpenColorIO/Private/OpenColorIODerivedDataVersion.h
+++ b/Engine/Plugins/Compositing/OpenColorIO/Source/OpenColorIO/Private/OpenColorIODerivedDataVersion.h
@@ -10,4 +10,4 @@ OpenColorIOShaderDerivedDataVersion.h: Derived data version for OpenColorIO shad
 
 // In case of merge conflicts with DDC versions, you MUST generate a new GUID and set this new
 // guid as version
-#define OPENCOLORIO_DERIVEDDATA_VER		TEXT("FECF9DF3-80B5-4551-B341-7244B95FDEE7")
\ No newline at end of file
+#define OPENCOLORIO_DERIVEDDATA_VER		TEXT("1AC059DA-E69A-498A-A904-C0DE6E27D314")
\ No newline at end of file
diff --git a/Engine/Plugins/FX/Niagara/Source/NiagaraShader/Private/NiagaraShaderDerivedDataVersion.h b/Engine/Plugins/FX/Niagara/Source/NiagaraShader/Private/NiagaraShaderDerivedDataVersion.h
--- a/Engine/Plugins/FX/Niagara/Source/NiagaraShader/Private/NiagaraShaderDerivedDataVersion.h
+++ b/Engine/Plugins/FX/Niagara/Source/NiagaraShader/Private/NiagaraShaderDerivedDataVersion.h
@@ -11,4 +11,4 @@ NiagaraShaderDerivedDataVersion.h: Shader derived data version for Niagara.
 // In case of merge conflicts with DDC versions, you MUST generate a new GUID and set this new
 // guid as version
 
-#define NIAGARASHADERMAP_DERIVEDDATA_VER		TEXT("3136AAAD-A9DE-4368-9F86-E89840E7CC87")
+#define NIAGARASHADERMAP_DERIVEDDATA_VER		TEXT("DF5C47C0-F984-42E7-8960-20F53987595B")
diff --git a/Engine/Shaders/Public/ShaderVersion.ush b/Engine/Shaders/Public/ShaderVersion.ush
--- a/Engine/Shaders/Public/ShaderVersion.ush
+++ b/Engine/Shaders/Public/ShaderVersion.ush
@@ -3,4 +3,4 @@
 // in Platform.ush (which should be included in any shader) it allows to invalidate the shader DDC.
 // 
 // If you are merging streams and there is a conflict with this GUID you should make a new GUID rather than taking one or the other.
-// GUID = F0FE2A62-52BB-48E7-8FEB-0E8766DA3D1D
+// GUID = 6ECB96A4-C8C9-4BA1-AB45-BECDB44269BF
\ No newline at end of file
