FHashedName converts strings to upper-case before hashing, otherwise hash depends on what particular case of a given string is first seen

diff --git a/Engine/Plugins/Compositing/OpenColorIO/Source/OpenColorIO/Private/OpenColorIODerivedDataVersion.h b/Engine/Plugins/Compositing/OpenColorIO/Source/OpenColorIO/Private/OpenColorIODerivedDataVersion.h
--- a/Engine/Plugins/Compositing/OpenColorIO/Source/OpenColorIO/Private/OpenColorIODerivedDataVersion.h
+++ b/Engine/Plugins/Compositing/OpenColorIO/Source/OpenColorIO/Private/OpenColorIODerivedDataVersion.h
@@ -10,4 +10,4 @@ OpenColorIOShaderDerivedDataVersion.h: Derived data version for OpenColorIO shad
 
 // In case of merge conflicts with DDC versions, you MUST generate a new GUID and set this new
 // guid as version
-#define OPENCOLORIO_DERIVEDDATA_VER		TEXT("1AC059DA-E69A-498A-A904-C0DE6E27D314")
\ No newline at end of file
+#define OPENCOLORIO_DERIVEDDATA_VER		TEXT("78E88333-7736-4D29-87B7-C21E8535234A")
\ No newline at end of file
diff --git a/Engine/Plugins/FX/Niagara/Source/NiagaraShader/Private/NiagaraShaderDerivedDataVersion.h b/Engine/Plugins/FX/Niagara/Source/NiagaraShader/Private/NiagaraShaderDerivedDataVersion.h
--- a/Engine/Plugins/FX/Niagara/Source/NiagaraShader/Private/NiagaraShaderDerivedDataVersion.h
+++ b/Engine/Plugins/FX/Niagara/Source/NiagaraShader/Private/NiagaraShaderDerivedDataVersion.h
@@ -11,4 +11,4 @@ NiagaraShaderDerivedDataVersion.h: Shader derived data version for Niagara.
 // In case of merge conflicts with DDC versions, you MUST generate a new GUID and set this new
 // guid as version
 
-#define NIAGARASHADERMAP_DERIVEDDATA_VER		TEXT("DF5C47C0-F984-42E7-8960-20F53987595B")
+#define NIAGARASHADERMAP_DERIVEDDATA_VER		TEXT("C082A601-ECB4-4561-AED5-D131CB079201")
diff --git a/Engine/Shaders/Public/ShaderVersion.ush b/Engine/Shaders/Public/ShaderVersion.ush
--- a/Engine/Shaders/Public/ShaderVersion.ush
+++ b/Engine/Shaders/Public/ShaderVersion.ush
@@ -3,4 +3,4 @@
 // in Platform.ush (which should be included in any shader) it allows to invalidate the shader DDC.
 // 
 // If you are merging streams and there is a conflict with this GUID you should make a new GUID rather than taking one or the other.
-// GUID = 6ECB96A4-C8C9-4BA1-AB45-BECDB44269BF
\ No newline at end of file
+// GUID = 90D1B596-2512-4A90-9FE2-F88BD15F9E7A
\ No newline at end of file
