Fix local cooking on Mac.

#rnx

diff --git a/Engine/Plugins/Compositing/OpenColorIO/Source/OpenColorIO/Private/OpenColorIOColorTransform.cpp b/Engine/Plugins/Compositing/OpenColorIO/Source/OpenColorIO/Private/OpenColorIOColorTransform.cpp
--- a/Engine/Plugins/Compositing/OpenColorIO/Source/OpenColorIO/Private/OpenColorIOColorTransform.cpp
+++ b/Engine/Plugins/Compositing/OpenColorIO/Source/OpenColorIO/Private/OpenColorIOColorTransform.cpp
@@ -414,10 +414,12 @@ bool UOpenColorIOColorTransform::UpdateShaderInfo(FString& OutShaderCodeHash, FS
 	return false;
 #else
 	//Avoid triggering errors when building maps on build machine.
+#if PLATFORM_WINDOWS
 	if (!GIsBuildMachine)
 	{
 		UE_LOG(LogOpenColorIO, Error, TEXT("Can't update shader, OCIO library isn't present."));
 	}
+#endif //PLATFORM_WINDOWS
 	return false;
 #endif //WITH_OCIO
 #else
