Fix collision between OpenColorIO SM for Windows+WindowsNoEditor (UE-90852).

diff --git a/Engine/Plugins/Compositing/OpenColorIO/Source/OpenColorIO/Private/OpenColorIOColorTransform.cpp b/Engine/Plugins/Compositing/OpenColorIO/Source/OpenColorIO/Private/OpenColorIOColorTransform.cpp
--- a/Engine/Plugins/Compositing/OpenColorIO/Source/OpenColorIO/Private/OpenColorIOColorTransform.cpp
+++ b/Engine/Plugins/Compositing/OpenColorIO/Source/OpenColorIO/Private/OpenColorIOColorTransform.cpp
@@ -139,7 +139,7 @@ void UOpenColorIOColorTransform::Serialize(FArchive& Ar)
 }
 
 
-void UOpenColorIOColorTransform::CacheResourceShadersForCooking(EShaderPlatform InShaderPlatform, const FString& InShaderHash, const FString& InShaderCode, TArray<FOpenColorIOTransformResource*>& OutCachedResources)
+void UOpenColorIOColorTransform::CacheResourceShadersForCooking(EShaderPlatform InShaderPlatform, const ITargetPlatform* TargetPlatform, const FString& InShaderHash, const FString& InShaderCode, TArray<FOpenColorIOTransformResource*>& OutCachedResources)
 {
 	const ERHIFeatureLevel::Type TargetFeatureLevel = GetMaxSupportedFeatureLevel(InShaderPlatform);
 
@@ -281,9 +281,9 @@ void UOpenColorIOColorTransform::CacheResourceShadersForRendering(bool bRegenera
 	}
 }
 
-void UOpenColorIOColorTransform::CacheShadersForResources(EShaderPlatform InShaderPlatform, FOpenColorIOTransformResource* InResourceToCache, bool bApplyCompletedShaderMapForRendering, bool bIsCooking)
+void UOpenColorIOColorTransform::CacheShadersForResources(EShaderPlatform InShaderPlatform, FOpenColorIOTransformResource* InResourceToCache, bool bApplyCompletedShaderMapForRendering, bool bIsCooking, const ITargetPlatform* TargetPlatform)
 {
-	const bool bSuccess = InResourceToCache->CacheShaders(InShaderPlatform, bApplyCompletedShaderMapForRendering, bIsCooking);
+	const bool bSuccess = InResourceToCache->CacheShaders(InShaderPlatform, TargetPlatform, bApplyCompletedShaderMapForRendering, bIsCooking);
 
 	if (!bSuccess)
 	{
@@ -560,7 +560,7 @@ void UOpenColorIOColorTransform::BeginCacheForCookedPlatformData(const ITargetPl
 				const EShaderPlatform LegacyShaderPlatform = ShaderFormatToLegacyShaderPlatform(DesiredShaderFormats[FormatIndex]);
 
 				// Begin caching shaders for the target platform and store the FOpenColorIOTransformResource being compiled into CachedColorTransformResourcesForCooking
-				CacheResourceShadersForCooking(LegacyShaderPlatform, ShaderCodeHash, ShaderCode, *CachedColorTransformResourceForPlatformPtr);
+				CacheResourceShadersForCooking(LegacyShaderPlatform, TargetPlatform, ShaderCodeHash, ShaderCode, *CachedColorTransformResourceForPlatformPtr);
 			}
 		}
 	}
diff --git a/Engine/Plugins/Compositing/OpenColorIO/Source/OpenColorIO/Private/OpenColorIODerivedDataVersion.h b/Engine/Plugins/Compositing/OpenColorIO/Source/OpenColorIO/Private/OpenColorIODerivedDataVersion.h
--- a/Engine/Plugins/Compositing/OpenColorIO/Source/OpenColorIO/Private/OpenColorIODerivedDataVersion.h
+++ b/Engine/Plugins/Compositing/OpenColorIO/Source/OpenColorIO/Private/OpenColorIODerivedDataVersion.h
@@ -10,4 +10,4 @@ OpenColorIOShaderDerivedDataVersion.h: Derived data version for OpenColorIO shad
 
 // In case of merge conflicts with DDC versions, you MUST generate a new GUID and set this new
 // guid as version
-#define OPENCOLORIO_DERIVEDDATA_VER		TEXT("78E88333-7736-4D29-87B7-C21E8535234A")
\ No newline at end of file
+#define OPENCOLORIO_DERIVEDDATA_VER		TEXT("BA838967-5A0D-42FE-A7E4-F02E28A61B52")
\ No newline at end of file
diff --git a/Engine/Plugins/Compositing/OpenColorIO/Source/OpenColorIO/Private/OpenColorIOShaderMap.cpp b/Engine/Plugins/Compositing/OpenColorIO/Source/OpenColorIO/Private/OpenColorIOShaderMap.cpp
--- a/Engine/Plugins/Compositing/OpenColorIO/Source/OpenColorIO/Private/OpenColorIOShaderMap.cpp
+++ b/Engine/Plugins/Compositing/OpenColorIO/Source/OpenColorIO/Private/OpenColorIOShaderMap.cpp
@@ -113,6 +113,11 @@ bool FOpenColorIOShaderMapId::operator==(const FOpenColorIOShaderMapId& InRefere
 		return false;
 	}
 
+	if (LayoutParams != InReferenceSet.LayoutParams)
+	{
+		return false;
+	}
+
 	for (int32 ShaderIndex = 0; ShaderIndex < ShaderTypeDependencies.Num(); ShaderIndex++)
 	{
 		const FShaderTypeDependency& ShaderTypeDependency = ShaderTypeDependencies[ShaderIndex];
@@ -129,9 +134,6 @@ bool FOpenColorIOShaderMapId::operator==(const FOpenColorIOShaderMapId& InRefere
 void FOpenColorIOShaderMapId::AppendKeyString(FString& OutKeyString) const
 {
 #if WITH_EDITOR
-	FPlatformTypeLayoutParameters LayoutParams;
-	LayoutParams.InitializeForCurrent();
-
 	OutKeyString += ShaderCodeHash;
 	OutKeyString += TEXT("_");
 
diff --git a/Engine/Plugins/Compositing/OpenColorIO/Source/OpenColorIO/Private/OpenColorIOShared.cpp b/Engine/Plugins/Compositing/OpenColorIO/Source/OpenColorIO/Private/OpenColorIOShared.cpp
--- a/Engine/Plugins/Compositing/OpenColorIO/Source/OpenColorIO/Private/OpenColorIOShared.cpp
+++ b/Engine/Plugins/Compositing/OpenColorIO/Source/OpenColorIO/Private/OpenColorIOShared.cpp
@@ -92,7 +92,7 @@ void FOpenColorIOTransformResource::GetDependentShaderTypes(EShaderPlatform InPl
 	OutShaderTypes.Sort(FCompareShaderTypes());
 }
 
-OPENCOLORIO_API void FOpenColorIOTransformResource::GetShaderMapId(EShaderPlatform InPlatform, FOpenColorIOShaderMapId& OutId) const
+OPENCOLORIO_API void FOpenColorIOTransformResource::GetShaderMapId(EShaderPlatform InPlatform, const ITargetPlatform* TargetPlatform, FOpenColorIOShaderMapId& OutId) const
 {
 	if (bLoadedCookedShaderMapId)
 	{
@@ -106,6 +106,22 @@ OPENCOLORIO_API void FOpenColorIOTransformResource::GetShaderMapId(EShaderPlatfo
 		OutId.FeatureLevel = GetFeatureLevel();
 		OutId.ShaderCodeHash = ShaderCodeHash;
 		OutId.SetShaderDependencies(ShaderTypes, InPlatform);
+#if WITH_EDITOR
+		if (TargetPlatform)
+		{
+			OutId.LayoutParams.InitializeForPlatform(TargetPlatform->IniPlatformName(), TargetPlatform->HasEditorOnlyData());
+		}
+		else
+		{
+			OutId.LayoutParams.InitializeForCurrent();
+		}
+#else
+		if (TargetPlatform != nullptr)
+		{
+			UE_LOG(LogMaterial, Error, TEXT("FOpenColorIOTransformResource::GetShaderMapId: TargetPlatform is not null, but a cooked executable cannot target platforms other than its own."));
+		}
+		OutId.LayoutParams.InitializeForCurrent();
+#endif
 	}
 }
 
@@ -206,10 +222,10 @@ OPENCOLORIO_API  bool FOpenColorIOTransformResource::IsCompilationFinished() con
 	return bRet;
 }
 
-bool FOpenColorIOTransformResource::CacheShaders(EShaderPlatform InPlatform, bool bApplyCompletedShaderMapForRendering, bool bSynchronous)
+bool FOpenColorIOTransformResource::CacheShaders(EShaderPlatform InPlatform, const ITargetPlatform* TargetPlatform, bool bApplyCompletedShaderMapForRendering, bool bSynchronous)
 {
 	FOpenColorIOShaderMapId ResourceShaderMapId;
-	GetShaderMapId(InPlatform, ResourceShaderMapId);
+	GetShaderMapId(InPlatform, TargetPlatform, ResourceShaderMapId);
 	return CacheShaders(ResourceShaderMapId, InPlatform, bApplyCompletedShaderMapForRendering, bSynchronous);
 }
 
diff --git a/Engine/Plugins/Compositing/OpenColorIO/Source/OpenColorIO/Public/OpenColorIOColorTransform.h b/Engine/Plugins/Compositing/OpenColorIO/Source/OpenColorIO/Public/OpenColorIOColorTransform.h
--- a/Engine/Plugins/Compositing/OpenColorIO/Source/OpenColorIO/Public/OpenColorIOColorTransform.h
+++ b/Engine/Plugins/Compositing/OpenColorIO/Source/OpenColorIO/Public/OpenColorIOColorTransform.h
@@ -41,7 +41,7 @@ public:
 	 * This does not apply completed results to the renderer scenes.
 	 * Caller is responsible for deleting OutCachedResource.
 	 */
-	void CacheResourceShadersForCooking(EShaderPlatform InShaderPlatform, const FString& InShaderHash, const FString& InShaderCode, TArray<FOpenColorIOTransformResource*>& OutCachedResources);
+	void CacheResourceShadersForCooking(EShaderPlatform InShaderPlatform, const ITargetPlatform* TargetPlatform, const FString& InShaderHash, const FString& InShaderCode, TArray<FOpenColorIOTransformResource*>& OutCachedResources);
 
 	/**
 	 * Serialize LUT data. This will effectively serialize the LUT only when cooking
@@ -59,7 +59,7 @@ public:
 	 * If a matching shader map is not found in memory or the DDC, a new one will be compiled.
 	 */
 	void CacheResourceShadersForRendering(bool bRegenerateId);
-	void CacheShadersForResources(EShaderPlatform InShaderPlatform, FOpenColorIOTransformResource* InResourcesToCache, bool bApplyCompletedShaderMapForRendering, bool bIsCooking);
+	void CacheShadersForResources(EShaderPlatform InShaderPlatform, FOpenColorIOTransformResource* InResourcesToCache, bool bApplyCompletedShaderMapForRendering, bool bIsCooking, const ITargetPlatform* TargetPlatform = nullptr);
 
 	FOpenColorIOTransformResource* AllocateResource();
 
diff --git a/Engine/Plugins/Compositing/OpenColorIO/Source/OpenColorIO/Public/OpenColorIOShared.h b/Engine/Plugins/Compositing/OpenColorIO/Source/OpenColorIO/Public/OpenColorIOShared.h
--- a/Engine/Plugins/Compositing/OpenColorIO/Source/OpenColorIO/Public/OpenColorIOShared.h
+++ b/Engine/Plugins/Compositing/OpenColorIO/Source/OpenColorIO/Public/OpenColorIOShared.h
@@ -53,6 +53,11 @@ public:
 	/** Shader types of shaders that are inlined in this shader map in the DDC. */
 	LAYOUT_FIELD(TMemoryImageArray<FShaderTypeDependency>, ShaderTypeDependencies);
 
+	/*
+	* Type layout parameters of the memory image
+	*/
+	LAYOUT_FIELD(FPlatformTypeLayoutParameters, LayoutParams);
+
 public:
 	FOpenColorIOShaderMapId()
 		: FeatureLevel(ERHIFeatureLevel::SM5)
@@ -315,7 +320,7 @@ public:
 	 * Caches the shaders for this color transform with no static parameters on the given platform.
 	 * This is used by UOpenColorIOColorTransform
 	 */
-	OPENCOLORIO_API  bool CacheShaders(EShaderPlatform InPlatform, bool bApplyCompletedShaderMapForRendering, bool bSynchronous);
+	OPENCOLORIO_API  bool CacheShaders(EShaderPlatform InPlatform, const ITargetPlatform* TargetPlatform, bool bApplyCompletedShaderMapForRendering, bool bSynchronous);
 	bool CacheShaders(const FOpenColorIOShaderMapId& InShaderMapId, EShaderPlatform InPlatform, bool bApplyCompletedShaderMapForRendering, bool bSynchronous);
 
 	/**
@@ -338,7 +343,7 @@ public:
 	OPENCOLORIO_API void DiscardShaderMap();
 
 	void GetDependentShaderTypes(EShaderPlatform InPlatform, TArray<FShaderType*>& OutShaderTypes) const;
-	OPENCOLORIO_API  virtual void GetShaderMapId(EShaderPlatform InPlatform, FOpenColorIOShaderMapId& OutId) const;
+	OPENCOLORIO_API  virtual void GetShaderMapId(EShaderPlatform InPlatform, const ITargetPlatform* TargetPlatform, FOpenColorIOShaderMapId& OutId) const;
 
 	OPENCOLORIO_API void Invalidate();
 
