Fix missing include
#rnx

diff --git a/Engine/Plugins/Compositing/OpenColorIO/Source/OpenColorIO/Private/OpenColorIOShared.cpp b/Engine/Plugins/Compositing/OpenColorIO/Source/OpenColorIO/Private/OpenColorIOShared.cpp
--- a/Engine/Plugins/Compositing/OpenColorIO/Source/OpenColorIO/Private/OpenColorIOShared.cpp
+++ b/Engine/Plugins/Compositing/OpenColorIO/Source/OpenColorIO/Private/OpenColorIOShared.cpp
@@ -19,6 +19,7 @@
 #include "UObject/CoreObjectVersion.h"
 #include "UObject/UObjectHash.h"
 #include "UObject/UObjectIterator.h"
+#include "Interfaces/ITargetPlatform.h"
 
 IMPLEMENT_TYPE_LAYOUT(FOpenColorIOCompilationOutput);
 IMPLEMENT_TYPE_LAYOUT(FOpenColorIOShaderMapId);
