Fix non-unity build (use different log category) (UE-92554).

diff --git a/Engine/Plugins/Compositing/OpenColorIO/Source/OpenColorIO/Private/OpenColorIOShared.cpp b/Engine/Plugins/Compositing/OpenColorIO/Source/OpenColorIO/Private/OpenColorIOShared.cpp
--- a/Engine/Plugins/Compositing/OpenColorIO/Source/OpenColorIO/Private/OpenColorIOShared.cpp
+++ b/Engine/Plugins/Compositing/OpenColorIO/Source/OpenColorIO/Private/OpenColorIOShared.cpp
@@ -119,7 +119,7 @@ OPENCOLORIO_API void FOpenColorIOTransformResource::GetShaderMapId(EShaderPlatfo
 #else
 		if (TargetPlatform != nullptr)
 		{
-			UE_LOG(LogMaterial, Error, TEXT("FOpenColorIOTransformResource::GetShaderMapId: TargetPlatform is not null, but a cooked executable cannot target platforms other than its own."));
+			UE_LOG(LogShaders, Error, TEXT("FOpenColorIOTransformResource::GetShaderMapId: TargetPlatform is not null, but a cooked executable cannot target platforms other than its own."));
 		}
 		OutId.LayoutParams.InitializeForCurrent();
 #endif
