- Fixing ocio assert in package project

diff --git a/Engine/Plugins/Compositing/OpenColorIO/Source/OpenColorIO/Private/OpenColorIOColorTransform.cpp b/Engine/Plugins/Compositing/OpenColorIO/Source/OpenColorIO/Private/OpenColorIOColorTransform.cpp
--- a/Engine/Plugins/Compositing/OpenColorIO/Source/OpenColorIO/Private/OpenColorIOColorTransform.cpp
+++ b/Engine/Plugins/Compositing/OpenColorIO/Source/OpenColorIO/Private/OpenColorIOColorTransform.cpp
@@ -163,11 +163,10 @@ void UOpenColorIOColorTransform::SerializeLuts(FArchive& Ar)
 
 	if (Ar.IsSaving())
 	{
-		UVolumeTexture* Lut3dTexturePtr = Lut3dTexture.Get();
 		int32 Num3dLutsToSave = 0;
 		if (Ar.IsCooking())
 		{
-			if (Lut3dTexturePtr != nullptr)
+			if (Lut3dTexture != nullptr)
 			{
 				Num3dLutsToSave = 1;
 			}
@@ -177,7 +176,7 @@ void UOpenColorIOColorTransform::SerializeLuts(FArchive& Ar)
 
 		if (Num3dLutsToSave > 0)
 		{
-			Ar << Lut3dTexturePtr;
+			Ar << Lut3dTexture;
 		}
 	}
 	else if (Ar.IsLoading())
@@ -188,16 +187,14 @@ void UOpenColorIOColorTransform::SerializeLuts(FArchive& Ar)
 		if (NumLoaded3dLuts > 0)
 		{
 			//Will only happen on cooked data
-			UVolumeTexture* TempTexture = nullptr;
-			Ar << TempTexture;
-			Lut3dTexture.Reset(TempTexture);
+			Ar << Lut3dTexture;
 		}
 	}
 }
 
 void UOpenColorIOColorTransform::CacheResourceTextures()
 {
-	if (!Lut3dTexture.IsValid())
+	if (Lut3dTexture == nullptr)
 	{
 #if WITH_EDITOR && WITH_OCIO
 		OCIO_NAMESPACE::ConstConfigRcPtr CurrentConfig = ConfigurationOwner->GetLoadedConfigurationFile();
@@ -320,10 +317,9 @@ bool UOpenColorIOColorTransform::GetShaderAndLUTResouces(ERHIFeatureLevel::Type
 	if (OutShaderResource)
 	{
 		//Some color transform will only require shader code with no LUT involved.
-		UVolumeTexture* Lut3dTexturePtr = Lut3dTexture.Get();
-		if (Lut3dTexturePtr != nullptr)
+		if (Lut3dTexture != nullptr)
 		{
-			OutLUT3dResource = Lut3dTexturePtr->Resource;
+			OutLUT3dResource = Lut3dTexture->Resource;
 		}
 
 		return true;
@@ -442,16 +438,15 @@ void UOpenColorIOColorTransform::Update3dLutTexture(const FString& InLutIdentifi
 #if WITH_EDITOR && WITH_OCIO
 	check(InSourceData);
 
-	Lut3dTexture.Reset(NewObject<UVolumeTexture>(this, NAME_None, RF_NoFlags));
-	UVolumeTexture* Lut3dTexturePtr = Lut3dTexture.Get();
+	Lut3dTexture = NewObject<UVolumeTexture>(this, NAME_None, RF_NoFlags);
 
 	//Initializes source data with the raw LUT. If it's found in DDC, the resulting platform data will be fetched from there. 
 	//If not, the source data will be used to generate the platform data.
-	Lut3dTexturePtr->MipGenSettings = TMGS_NoMipmaps;
-	Lut3dTexturePtr->CompressionNone = true;
-	Lut3dTexturePtr->Source.Init(OpenColorIOShader::Lut3dEdgeLength, OpenColorIOShader::Lut3dEdgeLength, OpenColorIOShader::Lut3dEdgeLength, /*NumMips=*/ 1, TSF_RGBA16F, nullptr);
+	Lut3dTexture->MipGenSettings = TMGS_NoMipmaps;
+	Lut3dTexture->CompressionNone = true;
+	Lut3dTexture->Source.Init(OpenColorIOShader::Lut3dEdgeLength, OpenColorIOShader::Lut3dEdgeLength, OpenColorIOShader::Lut3dEdgeLength, /*NumMips=*/ 1, TSF_RGBA16F, nullptr);
 
-	FFloat16Color* MipData = reinterpret_cast<FFloat16Color*>(Lut3dTexturePtr->Source.LockMip(0));
+	FFloat16Color* MipData = reinterpret_cast<FFloat16Color*>(Lut3dTexture->Source.LockMip(0));
 	const uint32 LutLength = OpenColorIOShader::Lut3dEdgeLength;
 	for (uint32 Z = 0; Z < LutLength; ++Z)
 	{
@@ -465,15 +460,15 @@ void UOpenColorIOColorTransform::Update3dLutTexture(const FString& InLutIdentifi
 			}
 		}
 	}
-	Lut3dTexturePtr->Source.UnlockMip(0);
+	Lut3dTexture->Source.UnlockMip(0);
 
 	//Generate a Guid from the identifier received from the library and our DDC version.
 	FGuid LutGuid;
 	GetOpenColorIOLUTKeyGuid(InLutIdentifier, LutGuid);
-	Lut3dTexturePtr->Source.SetId(LutGuid, true);
+	Lut3dTexture->Source.SetId(LutGuid, true);
 
 	//Process our new texture to be usable in rendering pipeline.
-	Lut3dTexturePtr->UpdateResource();
+	Lut3dTexture->UpdateResource();
 #endif
 }
 
diff --git a/Engine/Plugins/Compositing/OpenColorIO/Source/OpenColorIO/Public/OpenColorIOColorTransform.h b/Engine/Plugins/Compositing/OpenColorIO/Source/OpenColorIO/Public/OpenColorIOColorTransform.h
--- a/Engine/Plugins/Compositing/OpenColorIO/Source/OpenColorIO/Public/OpenColorIOColorTransform.h
+++ b/Engine/Plugins/Compositing/OpenColorIO/Source/OpenColorIO/Public/OpenColorIOColorTransform.h
@@ -10,7 +10,6 @@
 #include "OpenColorIOShared.h"
 #include "RenderCommandFence.h"
 #include "RHIDefinitions.h"
-#include "UObject/StrongObjectPtr.h"
 #include "UObject/WeakObjectPtrTemplates.h"
 
 #include "OpenColorIOColorTransform.generated.h"
@@ -159,7 +158,8 @@ private:
 
 	/** If the color space requires a 3dLUT, this will contains the data to do the transform */
 	/** Note: This will be serialized when cooking. Otherwhise, it relies on raw data of the library and what's on DDC */
-	TStrongObjectPtr<UVolumeTexture> Lut3dTexture;
+	UPROPERTY(Transient)
+	UVolumeTexture* Lut3dTexture = nullptr;
 	
 	/** Inline ColorTransform resources serialized from disk. To be processed on game thread in PostLoad. */
 	TArray<FOpenColorIOTransformResource> LoadedTransformResources;
