OpenColorIO: source/dest color space conversion settings (in viewport or property editor) are reset if no longer available in the configuration.

diff --git a/Engine/Plugins/Compositing/OpenColorIO/Source/OpenColorIO/Private/OpenColorIOColorSpace.cpp b/Engine/Plugins/Compositing/OpenColorIO/Source/OpenColorIO/Private/OpenColorIOColorSpace.cpp
--- a/Engine/Plugins/Compositing/OpenColorIO/Source/OpenColorIO/Private/OpenColorIOColorSpace.cpp
+++ b/Engine/Plugins/Compositing/OpenColorIO/Source/OpenColorIO/Private/OpenColorIOColorSpace.cpp
@@ -34,6 +34,13 @@ bool FOpenColorIOColorSpace::IsValid() const
 	return ColorSpaceIndex != INDEX_NONE && !ColorSpaceName.IsEmpty();
 }
 
+void FOpenColorIOColorSpace::Reset()
+{
+	ColorSpaceIndex = INDEX_NONE;
+	ColorSpaceName.Reset();
+	FamilyName.Reset();
+}
+
 FString FOpenColorIOColorSpace::GetFamilyNameAtDepth(int32 InDepth) const
 {
 	FString ReturnName;
@@ -84,3 +91,24 @@ bool FOpenColorIOColorConversionSettings::IsValid() const
 
 	return false;
 }
+
+void FOpenColorIOColorConversionSettings::ValidateColorSpaces()
+{
+	if (ConfigurationSource)
+	{
+		if (!ConfigurationSource->HasDesiredColorSpace(SourceColorSpace))
+		{
+			SourceColorSpace.Reset();
+		}
+
+		if (!ConfigurationSource->HasDesiredColorSpace(DestinationColorSpace))
+		{
+			DestinationColorSpace.Reset();
+		}
+	}
+	else
+	{
+		SourceColorSpace.Reset();
+		DestinationColorSpace.Reset();
+	}
+}
diff --git a/Engine/Plugins/Compositing/OpenColorIO/Source/OpenColorIO/Private/OpenColorIOConfiguration.cpp b/Engine/Plugins/Compositing/OpenColorIO/Source/OpenColorIO/Private/OpenColorIOConfiguration.cpp
--- a/Engine/Plugins/Compositing/OpenColorIO/Source/OpenColorIO/Private/OpenColorIOConfiguration.cpp
+++ b/Engine/Plugins/Compositing/OpenColorIO/Source/OpenColorIO/Private/OpenColorIOConfiguration.cpp
@@ -81,6 +81,11 @@ bool UOpenColorIOConfiguration::HasTransform(const FString& InSourceColorSpace,
 	return (TransformData != nullptr);
 }
 
+bool UOpenColorIOConfiguration::HasDesiredColorSpace(const FOpenColorIOColorSpace& ColorSpace) const
+{
+	return DesiredColorSpaces.Find(ColorSpace) != INDEX_NONE;
+}
+
 bool UOpenColorIOConfiguration::Validate() const
 {
 #if WITH_EDITOR 
diff --git a/Engine/Plugins/Compositing/OpenColorIO/Source/OpenColorIO/Private/OpenColorIODisplayManager.cpp b/Engine/Plugins/Compositing/OpenColorIO/Source/OpenColorIO/Private/OpenColorIODisplayManager.cpp
--- a/Engine/Plugins/Compositing/OpenColorIO/Source/OpenColorIO/Private/OpenColorIODisplayManager.cpp
+++ b/Engine/Plugins/Compositing/OpenColorIO/Source/OpenColorIO/Private/OpenColorIODisplayManager.cpp
@@ -6,14 +6,15 @@
 
 FOpenColorIODisplayConfiguration& FOpenColorIODisplayManager::FindOrAddDisplayConfiguration(FViewportClient* InViewportClient)
 {
-	FOpenColorIODisplayConfiguration* Config = nullptr;
 	const TSharedPtr<FOpenColorIODisplayExtension, ESPMode::ThreadSafe>* Extension = DisplayExtensions.FindByPredicate([InViewportClient](const TSharedPtr<FOpenColorIODisplayExtension, ESPMode::ThreadSafe>& Other) { return  Other.IsValid() && InViewportClient == Other->GetAssociatedViewportClient(); });
 	if (Extension)
 	{
 		FOpenColorIODisplayExtension* ExtensionPtr = Extension->Get();
 		check(ExtensionPtr);
 
-		return ExtensionPtr->GetDisplayConfiguration();
+		FOpenColorIODisplayConfiguration& DisplayConfiguration = ExtensionPtr->GetDisplayConfiguration();
+		DisplayConfiguration.ColorConfiguration.ValidateColorSpaces();
+		return DisplayConfiguration;
 	}
 	
 	// Extension not found, create it and return its config
@@ -30,7 +31,9 @@ const FOpenColorIODisplayConfiguration* FOpenColorIODisplayManager::GetDisplayCo
 		FOpenColorIODisplayExtension* ExtensionPtr = Extension->Get();
 		check(ExtensionPtr);
 
-		return &ExtensionPtr->GetDisplayConfiguration();
+		FOpenColorIODisplayConfiguration& DisplayConfiguration = ExtensionPtr->GetDisplayConfiguration();
+		DisplayConfiguration.ColorConfiguration.ValidateColorSpaces();
+		return &DisplayConfiguration;
 	}
 
 	return nullptr;
diff --git a/Engine/Plugins/Compositing/OpenColorIO/Source/OpenColorIO/Public/OpenColorIOColorSpace.h b/Engine/Plugins/Compositing/OpenColorIO/Source/OpenColorIO/Public/OpenColorIOColorSpace.h
--- a/Engine/Plugins/Compositing/OpenColorIO/Source/OpenColorIO/Public/OpenColorIOColorSpace.h
+++ b/Engine/Plugins/Compositing/OpenColorIO/Source/OpenColorIO/Public/OpenColorIOColorSpace.h
@@ -58,6 +58,9 @@ public:
 	/** Return true if the index and name have been set properly */
 	bool IsValid() const;
 
+	/** Reset members to default/empty values. */
+	void Reset();
+
 	/** 
 	 * Return the family name at the desired depth level 
 	 * @param InDepth Desired depth in the family string. 0 == First layer. 
@@ -81,7 +84,7 @@ public:
 
 	/** The source color space name. */
 	UPROPERTY(EditAnywhere, BlueprintReadWrite, Category = ColorSpace)
-	UOpenColorIOConfiguration* ConfigurationSource;
+	TObjectPtr<UOpenColorIOConfiguration> ConfigurationSource;
 
 	/** The source color space name. */
 	UPROPERTY(EditAnywhere, BlueprintReadWrite, Category = ColorSpace)
@@ -103,6 +106,11 @@ public:
 	 * Returns true if the source and destination color spaces are found in the configuration file
 	 */
 	bool IsValid() const;
+
+	/**
+	* Ensure that the selected source and destination color spaces are valid, resets them otherwise.
+	*/
+	void ValidateColorSpaces();
 };
 
 /**
diff --git a/Engine/Plugins/Compositing/OpenColorIO/Source/OpenColorIO/Public/OpenColorIOConfiguration.h b/Engine/Plugins/Compositing/OpenColorIO/Source/OpenColorIO/Public/OpenColorIOConfiguration.h
--- a/Engine/Plugins/Compositing/OpenColorIO/Source/OpenColorIO/Public/OpenColorIOConfiguration.h
+++ b/Engine/Plugins/Compositing/OpenColorIO/Source/OpenColorIO/Public/OpenColorIOConfiguration.h
@@ -37,6 +37,7 @@ public:
 public:
 	bool GetShaderAndLUTResources(ERHIFeatureLevel::Type InFeatureLevel, const FString& InSourceColorSpace, const FString& InDestinationColorSpace, FOpenColorIOTransformResource*& OutShaderResource, FTextureResource*& OutLUT3dResource);
 	bool HasTransform(const FString& InSourceColorSpace, const FString& InDestinationColorSpace);
+	bool HasDesiredColorSpace(const FOpenColorIOColorSpace& ColorSpace) const;
 	bool Validate() const;
 
 	/** This forces to reload colorspaces and corresponding shaders if those are not loaded already. */
diff --git a/Engine/Plugins/Compositing/OpenColorIO/Source/OpenColorIOEditor/Private/OpenColorIOColorSpaceConversionCustomization.cpp b/Engine/Plugins/Compositing/OpenColorIO/Source/OpenColorIOEditor/Private/OpenColorIOColorSpaceConversionCustomization.cpp
--- a/Engine/Plugins/Compositing/OpenColorIO/Source/OpenColorIOEditor/Private/OpenColorIOColorSpaceConversionCustomization.cpp
+++ b/Engine/Plugins/Compositing/OpenColorIO/Source/OpenColorIOEditor/Private/OpenColorIOColorSpaceConversionCustomization.cpp
@@ -51,6 +51,8 @@ void FOpenColorIOColorConversionSettingsCustomization::CustomizeHeader(TSharedRe
 							FOpenColorIOColorConversionSettings* Conversion = reinterpret_cast<FOpenColorIOColorConversionSettings*>(RawData[0]);
 							if (Conversion != nullptr)
 							{
+								Conversion->ValidateColorSpaces();
+
 								return FText::FromString(*Conversion->ToString());
 							}
 							else
