OpenColorIO: Invalidate color space selection (menu) when a different configuration is selected.

diff --git a/Engine/Plugins/Compositing/OpenColorIO/Source/OpenColorIOEditor/Private/Widgets/SOpenColorIOColorSpacePicker.cpp b/Engine/Plugins/Compositing/OpenColorIO/Source/OpenColorIOEditor/Private/Widgets/SOpenColorIOColorSpacePicker.cpp
--- a/Engine/Plugins/Compositing/OpenColorIO/Source/OpenColorIOEditor/Private/Widgets/SOpenColorIOColorSpacePicker.cpp
+++ b/Engine/Plugins/Compositing/OpenColorIO/Source/OpenColorIOEditor/Private/Widgets/SOpenColorIOColorSpacePicker.cpp
@@ -74,6 +74,9 @@ void SOpenColorIOColorSpacePicker::Construct(const FArguments& InArgs)
 void SOpenColorIOColorSpacePicker::SetConfiguration(TWeakObjectPtr<UOpenColorIOConfiguration> NewConfiguration)
 {
 	Configuration = NewConfiguration;
+
+	// Invalidate current color space selection
+	SetCurrentColorSpace(FOpenColorIOColorSpace());
 }
 
 void SOpenColorIOColorSpacePicker::SetRestrictedColorSpace(const FOpenColorIOColorSpace& InRestrictedColorSpace)
