OCIO: Fix incorrect use of TCHAR_TO_ANSI where pointer to out-of-scope temporary object was held, causing visible issues on other platforms. Now using StringCast instead.

diff --git a/Engine/Plugins/Compositing/OpenColorIO/Source/OpenColorIO/Private/OpenColorIOConfiguration.cpp b/Engine/Plugins/Compositing/OpenColorIO/Source/OpenColorIO/Private/OpenColorIOConfiguration.cpp
--- a/Engine/Plugins/Compositing/OpenColorIO/Source/OpenColorIO/Private/OpenColorIOConfiguration.cpp
+++ b/Engine/Plugins/Compositing/OpenColorIO/Source/OpenColorIO/Private/OpenColorIOConfiguration.cpp
@@ -122,7 +122,7 @@ void UOpenColorIOConfiguration::ReloadExistingColorspaces()
 	// This will make sure that all colorspaces are up to date in case an index, family or name is changed.
 	for (const FOpenColorIOColorSpace& ExistingColorSpace : ColorSpacesToBeReloaded)
 	{
-		char* ColorSpaceName = TCHAR_TO_ANSI(*ExistingColorSpace.ColorSpaceName);
+		const char* ColorSpaceName = StringCast<ANSICHAR>(*ExistingColorSpace.ColorSpaceName).Get();
 		int ColorSpaceIndex = LoadedConfig->getIndexForColorSpace(ColorSpaceName);
 		OCIO_NAMESPACE::ConstColorSpaceRcPtr LibColorSpace = LoadedConfig->getColorSpace(ColorSpaceName);
 		if (!LibColorSpace)
@@ -133,7 +133,7 @@ void UOpenColorIOConfiguration::ReloadExistingColorspaces()
 
 		FOpenColorIOColorSpace ColorSpace;
 		ColorSpace.ColorSpaceIndex = ColorSpaceIndex;
-		ColorSpace.ColorSpaceName = StringCast<TCHAR>(ColorSpaceName).Get();
+		ColorSpace.ColorSpaceName = ExistingColorSpace.ColorSpaceName;
 		ColorSpace.FamilyName = StringCast<TCHAR>(LibColorSpace->getFamily()).Get();
 		DesiredColorSpaces.Add(ColorSpace);
 	}
