Fixed usage of StringCast<ANSICHAR> in OpenColorIO.

diff --git a/Engine/Plugins/Compositing/OpenColorIO/Source/OpenColorIO/Private/OpenColorIOConfiguration.cpp b/Engine/Plugins/Compositing/OpenColorIO/Source/OpenColorIO/Private/OpenColorIOConfiguration.cpp
--- a/Engine/Plugins/Compositing/OpenColorIO/Source/OpenColorIO/Private/OpenColorIOConfiguration.cpp
+++ b/Engine/Plugins/Compositing/OpenColorIO/Source/OpenColorIO/Private/OpenColorIOConfiguration.cpp
@@ -122,9 +122,9 @@ void UOpenColorIOConfiguration::ReloadExistingColorspaces()
 	// This will make sure that all colorspaces are up to date in case an index, family or name is changed.
 	for (const FOpenColorIOColorSpace& ExistingColorSpace : ColorSpacesToBeReloaded)
 	{
-		const char* ColorSpaceName = StringCast<ANSICHAR>(*ExistingColorSpace.ColorSpaceName).Get();
-		int ColorSpaceIndex = LoadedConfig->getIndexForColorSpace(ColorSpaceName);
-		OCIO_NAMESPACE::ConstColorSpaceRcPtr LibColorSpace = LoadedConfig->getColorSpace(ColorSpaceName);
+		const auto ColorSpaceName = StringCast<ANSICHAR>(*ExistingColorSpace.ColorSpaceName);
+		int ColorSpaceIndex = LoadedConfig->getIndexForColorSpace(ColorSpaceName.Get());
+		OCIO_NAMESPACE::ConstColorSpaceRcPtr LibColorSpace = LoadedConfig->getColorSpace(ColorSpaceName.Get());
 		if (!LibColorSpace)
 		{
 			// Name not found, therefore we don't need to re-add this colorspace.
