- Fixing unsafe property accesses  in ocio customization

diff --git a/Engine/Plugins/Compositing/OpenColorIO/Source/OpenColorIOEditor/Private/OpenColorIOColorSpaceConversionCustomization.cpp b/Engine/Plugins/Compositing/OpenColorIO/Source/OpenColorIOEditor/Private/OpenColorIOColorSpaceConversionCustomization.cpp
--- a/Engine/Plugins/Compositing/OpenColorIO/Source/OpenColorIOEditor/Private/OpenColorIOColorSpaceConversionCustomization.cpp
+++ b/Engine/Plugins/Compositing/OpenColorIO/Source/OpenColorIOEditor/Private/OpenColorIOColorSpaceConversionCustomization.cpp
@@ -46,20 +46,23 @@ void FOpenColorIOColorConversionSettingsCustomization::CustomizeHeader(TSharedRe
 						SNew(STextBlock)
 						.Text(MakeAttributeLambda([=]
 						{
-							TArray<void*> RawData;
-							ColorConversionProperty->AccessRawData(RawData);
-							FOpenColorIOColorConversionSettings* Conversion = reinterpret_cast<FOpenColorIOColorConversionSettings*>(RawData[0]);
-							if (Conversion != nullptr)
+							if (ColorConversionProperty.IsValid())
 							{
-								Conversion->ValidateColorSpaces();
-
-								return FText::FromString(*Conversion->ToString());
-							}
-							else
-							{
-								return FText::FromString(TEXT("<Invalid Conversion>"));
+								TArray<void*> RawData;
+								ColorConversionProperty->AccessRawData(RawData);
+								if (RawData.Num() > 0)
+								{
+									FOpenColorIOColorConversionSettings* Conversion = reinterpret_cast<FOpenColorIOColorConversionSettings*>(RawData[0]);
+									if (Conversion != nullptr)
+									{
+										Conversion->ValidateColorSpaces();
+										return FText::FromString(*Conversion->ToString());
+									}
+								}
 							}
 
+							return FText::FromString(TEXT("<Invalid Conversion>"));
+
 						}))
 					]
 				].IsEnabled(MakeAttributeLambda([=] { return !InPropertyHandle->IsEditConst() && PropertyUtils->IsPropertyEditingEnabled(); }));
@@ -69,10 +72,6 @@ void FOpenColorIOColorConversionSettingsCustomization::CustomizeHeader(TSharedRe
 
 void FOpenColorIOColorConversionSettingsCustomization::CustomizeChildren(TSharedRef<IPropertyHandle> InStructPropertyHandle, IDetailChildrenBuilder& StructBuilder, IPropertyTypeCustomizationUtils& StructCustomizationUtils)
 {
-	TArray<void*> RawData;
-	ColorConversionProperty->AccessRawData(RawData);
-	FOpenColorIOColorConversionSettings* ColorSpaceValue = reinterpret_cast<FOpenColorIOColorConversionSettings*>(RawData[0]);
-
 	TSharedPtr<IPropertyUtilities> PropertyUtils = StructCustomizationUtils.GetPropertyUtilities();
 
 	uint32 NumberOfChild;
@@ -130,15 +129,16 @@ void FOpenColorIOColorConversionSettingsCustomization::AddColorSpaceRow(FDetailW
 				{
 					TArray<void*> RawData;
 					InChildHandle->AccessRawData(RawData);
-					FOpenColorIOColorSpace* ColorSpace = reinterpret_cast<FOpenColorIOColorSpace*>(RawData[0]);
-					if (ColorSpace != nullptr)
-					{
-						return FText::FromString(*ColorSpace->ToString());
-					}
-					else
+					if (RawData.Num() > 0)
 					{
-						return FText::FromString(TEXT("<Invalid>"));
+						FOpenColorIOColorSpace* ColorSpace = reinterpret_cast<FOpenColorIOColorSpace*>(RawData[0]);
+						if (ColorSpace != nullptr)
+						{
+							return FText::FromString(*ColorSpace->ToString());
+						}
 					}
+				
+					return FText::FromString(TEXT("<Invalid>"));
 				}))
 			]
 			+ SHorizontalBox::Slot()
@@ -163,9 +163,17 @@ void FOpenColorIOColorConversionSettingsCustomization::AddColorSpaceRow(FDetailW
 
 TSharedRef<SWidget> FOpenColorIOColorConversionSettingsCustomization::HandleColorSpaceComboButtonMenuContent(TSharedPtr<IPropertyHandle> InPropertyHandle) const
 {
-	TArray<void*> RawData;
-	ColorConversionProperty->AccessRawData(RawData);
-	FOpenColorIOColorConversionSettings* ColorSpaceConversion = reinterpret_cast<FOpenColorIOColorConversionSettings*>(RawData[0]);
+	FOpenColorIOColorConversionSettings* ColorSpaceConversion = nullptr;
+
+	if (InPropertyHandle.IsValid())
+	{
+		TArray<void*> RawData;
+		ColorConversionProperty->AccessRawData(RawData);
+		if (RawData.Num() > 0)
+		{
+			ColorSpaceConversion = reinterpret_cast<FOpenColorIOColorConversionSettings*>(RawData[0]);
+		}
+	}
 
 	if (ColorSpaceConversion && ColorSpaceConversion->ConfigurationSource)
 	{
@@ -203,27 +211,37 @@ TSharedRef<SWidget> FOpenColorIOColorConversionSettingsCustomization::HandleColo
 					(
 						FExecuteAction::CreateLambda([=] 
 						{
-							if (FStructProperty* StructProperty = CastField<FStructProperty>(InPropertyHandle->GetProperty()))
+							if (InPropertyHandle.IsValid())
 							{
-								TArray<void*> RawData;
-								InPropertyHandle->AccessRawData(RawData);
-								FOpenColorIOColorSpace* PreviousColorSpaceValue = reinterpret_cast<FOpenColorIOColorSpace*>(RawData[0]);
+								if (FStructProperty* StructProperty = CastField<FStructProperty>(InPropertyHandle->GetProperty()))
+								{
+									TArray<void*> RawData;
+									InPropertyHandle->AccessRawData(RawData);
+									if (RawData.Num() > 0)
+									{
+										FOpenColorIOColorSpace* PreviousColorSpaceValue = reinterpret_cast<FOpenColorIOColorSpace*>(RawData[0]);
 
-								FString TextValue;
-								StructProperty->Struct->ExportText(TextValue, &ColorSpace, PreviousColorSpaceValue, nullptr, EPropertyPortFlags::PPF_None, nullptr);
-								ensure(InPropertyHandle->SetValueFromFormattedString(TextValue, EPropertyValueSetFlags::DefaultFlags) == FPropertyAccess::Result::Success);
+										FString TextValue;
+										StructProperty->Struct->ExportText(TextValue, &ColorSpace, PreviousColorSpaceValue, nullptr, EPropertyPortFlags::PPF_None, nullptr);
+										ensure(InPropertyHandle->SetValueFromFormattedString(TextValue, EPropertyValueSetFlags::DefaultFlags) == FPropertyAccess::Result::Success);
+									}
+								}
 							}
 						}),
 						FCanExecuteAction(),
 						FIsActionChecked::CreateLambda([=] 
 						{
-							TArray<void*> RawData;
-							SourceColorSpaceProperty->AccessRawData(RawData);
-							if (RawData.Num() > 0)
+							if (SourceColorSpaceProperty.IsValid())
 							{
-								FOpenColorIOColorSpace* PreviousColorSpaceValue = reinterpret_cast<FOpenColorIOColorSpace*>(RawData[0]);
-								return *PreviousColorSpaceValue == ColorSpace;
+								TArray<void*> RawData;
+								SourceColorSpaceProperty->AccessRawData(RawData);
+								if (RawData.Num() > 0)
+								{
+									FOpenColorIOColorSpace* PreviousColorSpaceValue = reinterpret_cast<FOpenColorIOColorSpace*>(RawData[0]);
+									return *PreviousColorSpaceValue == ColorSpace;
+								}
 							}
+						
 							return false;
 						})
 					),
