OCIO: Prevent rare lumen-related viewport client crash.

#rnx

diff --git a/Engine/Plugins/Compositing/OpenColorIO/Source/OpenColorIOEditor/Private/OpenColorIOEditorModule.cpp b/Engine/Plugins/Compositing/OpenColorIO/Source/OpenColorIOEditor/Private/OpenColorIOEditorModule.cpp
--- a/Engine/Plugins/Compositing/OpenColorIO/Source/OpenColorIOEditor/Private/OpenColorIOEditorModule.cpp
+++ b/Engine/Plugins/Compositing/OpenColorIO/Source/OpenColorIOEditor/Private/OpenColorIOEditorModule.cpp
@@ -290,14 +290,15 @@ void FOpenColorIOEditorModule::TrackNewViewportIfRequired(FViewport* Viewport)
 		FLevelEditorViewportClient* const* AssociatedClient = LevelViewportClients.FindByPredicate([Viewport](const FLevelEditorViewportClient* Other) { return Other == Viewport->GetClient(); });
 
 		// Active viewport should always have a client
-		check(AssociatedClient && *AssociatedClient);
-		
-		FLevelEditorViewportClient* Client = *AssociatedClient;
-		TSharedPtr<SLevelViewport> LevelViewport = StaticCastSharedPtr<SLevelViewport>(Client->GetEditorViewportWidget());
-		if (LevelViewport.IsValid())
+		if (ensure(AssociatedClient))
 		{
-			FViewportPair NewEntry(Client, LevelViewport->GetConfigKey());
-			ConfiguredViewports.Emplace(MoveTemp(NewEntry));
+			FLevelEditorViewportClient* Client = *AssociatedClient;
+			TSharedPtr<SLevelViewport> LevelViewport = StaticCastSharedPtr<SLevelViewport>(Client->GetEditorViewportWidget());
+			if (LevelViewport.IsValid())
+			{
+				FViewportPair NewEntry(Client, LevelViewport->GetConfigKey());
+				ConfiguredViewports.Emplace(MoveTemp(NewEntry));
+			}
 		}
 	}
 }
