OCIO: Convert all passes to RDG, removing the need for the older shader class.

diff --git a/Engine/Plugins/Compositing/OpenColorIO/Shaders/Private/OpenColorIOBaseVS.usf b/Engine/Plugins/Compositing/OpenColorIO/Shaders/Private/OpenColorIOBaseVS.usf
deleted file mode 100644
diff --git a/Engine/Plugins/Compositing/OpenColorIO/Shaders/Private/OpenColorIOShader.usf b/Engine/Plugins/Compositing/OpenColorIO/Shaders/Private/OpenColorIOShader.usf
--- a/Engine/Plugins/Compositing/OpenColorIO/Shaders/Private/OpenColorIOShader.usf
+++ b/Engine/Plugins/Compositing/OpenColorIO/Shaders/Private/OpenColorIOShader.usf
@@ -32,11 +32,10 @@ float Gamma;
 /* Pixel shader
 =============================================================================*/
 
-void MainPS(float4 InPosition : SV_POSITION,
-	float2 InUV : TEXCOORD0,
+void MainPS(in noperspective float4 InUVAndScreenPos : TEXCOORD0,
 	out float4 OutColor : SV_Target0)
 {
-	float4 InputColor = InputTexture.Sample(InputTextureSampler, InUV.xy);
+	float4 InputColor = InputTexture.Sample(InputTextureSampler, InUVAndScreenPos.xy);
 	InputColor = max(0.0, InputColor);
 	InputColor = pow(InputColor, Gamma);
 	OutColor = float4(OCIOConvert(InputColor).xyz, InputColor.w);
diff --git a/Engine/Plugins/Compositing/OpenColorIO/Source/OpenColorIO/Private/OpenColorIODisplayExtension.cpp b/Engine/Plugins/Compositing/OpenColorIO/Source/OpenColorIO/Private/OpenColorIODisplayExtension.cpp
--- a/Engine/Plugins/Compositing/OpenColorIO/Source/OpenColorIO/Private/OpenColorIODisplayExtension.cpp
+++ b/Engine/Plugins/Compositing/OpenColorIO/Source/OpenColorIO/Private/OpenColorIODisplayExtension.cpp
@@ -14,7 +14,6 @@
 #include "PipelineStateCache.h"
 #include "RenderTargetPool.h"
 #include "Shader.h"
-#include "OpenColorIORenderingPrivate.h"
 
 #include "RHI.h"
 #include "SceneView.h"
@@ -79,7 +78,7 @@ void FOpenColorIODisplayExtension::SetupView(FSceneViewFamily& InViewFamily, FSc
 		{
 			// Transform was found, so shader must be there but doesn't mean the actual shader is available
 			check(ShaderResource);
-			if (ShaderResource->GetShaderGameThread<FOpenColorIOPixelShader_RDG>().IsNull())
+			if (ShaderResource->GetShaderGameThread<FOpenColorIOPixelShader>().IsNull())
 			{
 				UE_LOG(LogOpenColorIO, Warning, TEXT("Can't apply display look - Shader was invalid for Resource %s"), *ShaderResource->GetFriendlyName());
 
@@ -109,48 +108,6 @@ void FOpenColorIODisplayExtension::SetupView(FSceneViewFamily& InViewFamily, FSc
 	);
 }
 
-namespace {
-	template<typename TSetupFunction>
-	void DrawScreenPass(
-		FRHICommandList& RHICmdList,
-		const FSceneView& View,
-		const FScreenPassTextureViewport& OutputViewport,
-		const FScreenPassTextureViewport& InputViewport,
-		const FScreenPassPipelineState& PipelineState,
-		TSetupFunction SetupFunction)
-	{
-		PipelineState.Validate();
-
-		const FIntRect InputRect = InputViewport.Rect;
-		const FIntPoint InputSize = InputViewport.Extent;
-		const FIntRect OutputRect = OutputViewport.Rect;
-		const FIntPoint OutputSize = OutputRect.Size();
-
-		RHICmdList.SetViewport(OutputRect.Min.X, OutputRect.Min.Y, 0.0f, OutputRect.Max.X, OutputRect.Max.Y, 1.0f);
-
-		SetScreenPassPipelineState(RHICmdList, PipelineState);
-
-		// Setting up buffers.
-		SetupFunction(RHICmdList);
-
-		FIntPoint LocalOutputPos(FIntPoint::ZeroValue);
-		FIntPoint LocalOutputSize(OutputSize);
-		EDrawRectangleFlags DrawRectangleFlags = EDRF_UseTriangleOptimization;
-
-		DrawPostProcessPass(
-			RHICmdList,
-			LocalOutputPos.X, LocalOutputPos.Y, LocalOutputSize.X, LocalOutputSize.Y,
-			InputRect.Min.X, InputRect.Min.Y, InputRect.Width(), InputRect.Height(),
-			OutputSize,
-			InputSize,
-			PipelineState.VertexShader,
-			View.StereoViewIndex,
-			false,
-			DrawRectangleFlags);
-	}
-
-}
-
 void FOpenColorIODisplayExtension::SubscribeToPostProcessingPass(EPostProcessingPass PassId, FAfterPassCallbackDelegateArray& InOutPassCallbacks, bool bIsPassEnabled)
 {
 	if (PassId == EPostProcessingPass::Tonemap)
@@ -164,8 +121,6 @@ FScreenPassTexture FOpenColorIODisplayExtension::PostProcessPassAfterTonemap_Ren
 
 	const FSceneViewFamily& ViewFamily = *View.Family;
 
-	const ERHIFeatureLevel::Type FeatureLevel = View.GetFeatureLevel();
-
 	const FScreenPassTexture& SceneColor = InOutInputs.Textures[(uint32)EPostProcessMaterialInput::SceneColor];
 
 	//If the shader resource could not be found, skip this frame. The LUT isn't required
@@ -186,88 +141,53 @@ FScreenPassTexture FOpenColorIODisplayExtension::PostProcessPassAfterTonemap_Ren
 
 	RDG_EVENT_SCOPE(GraphBuilder, "OCIODisplayLook");
 
+	FScreenPassRenderTarget BackBufferRenderTarget;
+
+	// If the override output is provided it means that this is the last pass in post processing.
+	if (InOutInputs.OverrideOutput.IsValid())
 	{
-		// Get shader from shader map.
-		FGlobalShaderMap* GlobalShaderMap = GetGlobalShaderMap(ViewFamily.GetFeatureLevel());
+		BackBufferRenderTarget = InOutInputs.OverrideOutput;
+	}
+	else
+	{
+		// Reusing the same output description for our back buffer as SceneColor when it's not overriden
+		FRDGTextureDesc OutputDesc = SceneColor.Texture->Desc;
+		OutputDesc.Flags |= TexCreate_RenderTargetable;
+		FLinearColor ClearColor(0., 0., 0., 0.);
+		OutputDesc.ClearValue = FClearValueBinding(ClearColor);
+
+		FRDGTexture* BackBufferRenderTargetTexture = GraphBuilder.CreateTexture(OutputDesc, TEXT("BackBufferRenderTargetTexture"));
+		BackBufferRenderTarget = FScreenPassRenderTarget(BackBufferRenderTargetTexture, SceneColor.ViewRect, ERenderTargetLoadAction::EClear);
+	}
 
-		
-		FScreenPassRenderTarget BackBufferRenderTarget;
+	//Get input and output viewports. Backbuffer could be targeting a different region than input viewport
+	const FScreenPassTextureViewport SceneColorViewport(SceneColor);
+	const FScreenPassTextureViewport BackBufferViewport(BackBufferRenderTarget);
 
-		// If the override output is provided it means that this is the last pass in post processing.
-		if (InOutInputs.OverrideOutput.IsValid())
-		{
-			BackBufferRenderTarget = InOutInputs.OverrideOutput;
-		}
-		else
-		{
-			// Reusing the same output description for our back buffer as SceneColor when it's not overriden
-			FRDGTextureDesc OutputDesc = SceneColor.Texture->Desc;
-			OutputDesc.Flags |= TexCreate_RenderTargetable;
-			FLinearColor ClearColor(0., 0., 0., 0.);
-			OutputDesc.ClearValue = FClearValueBinding(ClearColor);
-
-			FRDGTexture* BackBufferRenderTargetTexture = GraphBuilder.CreateTexture(OutputDesc, TEXT("BackBufferRenderTargetTexture"));
-			BackBufferRenderTarget = FScreenPassRenderTarget(BackBufferRenderTargetTexture, SceneColor.ViewRect, ERenderTargetLoadAction::EClear);
-		}
+	FScreenPassRenderTarget SceneColorRenderTarget(SceneColor, ERenderTargetLoadAction::ELoad);
 
-		//Get input and output viewports. Backbuffer could be targeting a different region than input viewport
-		const FScreenPassTextureViewport SceneColorViewport(SceneColor);
-		const FScreenPassTextureViewport BackBufferViewport(BackBufferRenderTarget);
+	TShaderRef<FOpenColorIOPixelShader> OCIOPixelShader = CachedResourcesRenderThread.ShaderResource->GetShader<FOpenColorIOPixelShader>();
 
-		FScreenPassRenderTarget SceneColorRenderTarget(SceneColor, ERenderTargetLoadAction::ELoad);
+	const float DisplayGamma = View.Family->RenderTarget->GetDisplayGamma();
 
-		FRHIBlendState* DefaultBlendState = FScreenPassPipelineState::FDefaultBlendState::GetRHI();
-		FRHIDepthStencilState* DepthStencilState = FScreenPassPipelineState::FDefaultDepthStencilState::GetRHI();
+	FOpenColorIOPixelShaderParameters* Parameters = GraphBuilder.AllocParameters<FOpenColorIOPixelShaderParameters>();
+	Parameters->InputTexture = SceneColorRenderTarget.Texture;
+	Parameters->InputTextureSampler = TStaticSamplerState<>::GetRHI();
+	if (CachedResourcesRenderThread.LUT3dResource)
+	{
+		Parameters->Ocio_lut3d_0 = CachedResourcesRenderThread.LUT3dResource->TextureRHI;
+	}
+	Parameters->Ocio_lut3d_0Sampler = TStaticSamplerState<SF_Bilinear, AM_Clamp, AM_Clamp, AM_Clamp>::GetRHI();
 
-		{
-			
-			TShaderMapRef<FOpenColorIOVertexShader> VertexShader(GlobalShaderMap);
-			TShaderRef<FOpenColorIOPixelShader_RDG> OCIOPixelShader = CachedResourcesRenderThread.ShaderResource->GetShaderGameThread<FOpenColorIOPixelShader_RDG>();
+	// There is a special case where post processing and tonemapper are disabled. In this case tonemapper applies a static display Inverse of Gamma which defaults to 2.2.
+	// In the case when Both PostProcessing and ToneMapper are disabled we apply gamma manually. In every other case we apply inverse gamma before applying OCIO.
+	Parameters->Gamma = (ViewFamily.EngineShowFlags.Tonemapper == 0) || (ViewFamily.EngineShowFlags.PostProcessing == 0) ? DefaultDisplayGamma : DefaultDisplayGamma / DisplayGamma;
+	Parameters->RenderTargets[0] = BackBufferRenderTarget.GetRenderTargetBinding();
 
-			const float DisplayGamma = View.Family->RenderTarget->GetDisplayGamma();
+	checkSlow(View.bIsViewInfo);
+	const FViewInfo& ViewInfo = static_cast<const FViewInfo&>(View);
 
-			FOpenColorIOPixelShaderParameters* Parameters = GraphBuilder.AllocParameters<FOpenColorIOPixelShaderParameters>();
-			Parameters->InputTexture = SceneColorRenderTarget.Texture;
-			Parameters->InputTextureSampler = TStaticSamplerState<>::GetRHI();
-			if (CachedResourcesRenderThread.LUT3dResource)
-			{
-				Parameters->Ocio_lut3d_0 = CachedResourcesRenderThread.LUT3dResource->TextureRHI;
-			}
-			Parameters->Ocio_lut3d_0Sampler = TStaticSamplerState<SF_Bilinear, AM_Clamp, AM_Clamp, AM_Clamp>::GetRHI();
-
-			// There is a special case where post processing and tonemapper are disabled. In this case tonemapper applies a static display Inverse of Gamma which defaults to 2.2.
-			// In the case when Both PostProcessing and ToneMapper are disabled we apply gamma manually. In every other case we apply inverse gamma before applying OCIO.
-			Parameters->Gamma = (ViewFamily.EngineShowFlags.Tonemapper == 0) || (ViewFamily.EngineShowFlags.PostProcessing == 0) ? DefaultDisplayGamma : DefaultDisplayGamma/DisplayGamma;
-			Parameters->RenderTargets[0] = BackBufferRenderTarget.GetRenderTargetBinding();
-
-			// Main Pass
-			GraphBuilder.AddPass(
-				RDG_EVENT_NAME("ProcessOCIOColorSpaceXfrm"),
-				Parameters,
-				ERDGPassFlags::Raster,
-				[&View,
-				VertexShader,
-				OCIOPixelShader,
-				DefaultBlendState,
-				DepthStencilState,
-				SceneColorViewport,
-				BackBufferViewport,
-				Parameters](FRHICommandList& RHICmdList)
-			{
-				DrawScreenPass(
-					RHICmdList,
-					View,
-					BackBufferViewport,
-					SceneColorViewport,
-					FScreenPassPipelineState(VertexShader, OCIOPixelShader, DefaultBlendState, DepthStencilState),
-					[&](FRHICommandList& RHICmdList)
-				{
-					SetShaderParameters(RHICmdList, OCIOPixelShader, OCIOPixelShader.GetPixelShader(), *Parameters);
-				});
-
-			});
-
-			return MoveTemp(BackBufferRenderTarget);
-		}
-	}
+	AddDrawScreenPass(GraphBuilder, RDG_EVENT_NAME("ProcessOCIOColorSpaceXfrm"), ViewInfo, BackBufferViewport, SceneColorViewport, OCIOPixelShader, Parameters);
+
+	return MoveTemp(BackBufferRenderTarget);
 }
\ No newline at end of file
diff --git a/Engine/Plugins/Compositing/OpenColorIO/Source/OpenColorIO/Private/OpenColorIORendering.cpp b/Engine/Plugins/Compositing/OpenColorIO/Source/OpenColorIO/Private/OpenColorIORendering.cpp
--- a/Engine/Plugins/Compositing/OpenColorIO/Source/OpenColorIO/Private/OpenColorIORendering.cpp
+++ b/Engine/Plugins/Compositing/OpenColorIO/Source/OpenColorIO/Private/OpenColorIORendering.cpp
@@ -6,7 +6,6 @@
 #include "Engine/TextureRenderTarget2D.h"
 #include "Engine/World.h"
 #include "GlobalShader.h"
-#include "CommonRenderResources.h"
 #include "Logging/LogMacros.h"
 #include "Logging/MessageLog.h"
 #include "OpenColorIOConfiguration.h"
@@ -14,46 +13,8 @@
 #include "OpenColorIOShader.h"
 #include "OpenColorIOShaderType.h"
 #include "OpenColorIOShared.h"
-#include "PipelineStateCache.h"
-#include "RHIStaticStates.h"
-#include "SceneInterface.h"
-#include "SceneUtils.h"
-#include "ShaderParameterUtils.h"
-#include "TextureResource.h"
 #include "ScreenPass.h"
-
-namespace {
-	/** This function is similar to DrawScreenPass in OpenColorIODisplayExtension.cpp except it is catered for Viewless texture rendering. */
-	template<typename TSetupFunction>
-	void DrawScreenPass(
-		FRHICommandListImmediate& RHICmdList,
-		const FIntPoint& OutputResolution,
-		const FScreenPassPipelineState& PipelineState,
-		TSetupFunction SetupFunction)
-	{
-		RHICmdList.SetViewport(0.f, 0.f, 0.f, OutputResolution.X, OutputResolution.Y, 1.0f);
-
-		SetScreenPassPipelineState(RHICmdList, PipelineState);
-
-		// Setting up buffers.
-		SetupFunction(RHICmdList);
-
-		FIntPoint LocalOutputPos(FIntPoint::ZeroValue);
-		FIntPoint LocalOutputSize(OutputResolution);
-		EDrawRectangleFlags DrawRectangleFlags = EDRF_UseTriangleOptimization;
-
-		DrawPostProcessPass(
-			RHICmdList,
-			LocalOutputPos.X, LocalOutputPos.Y, LocalOutputSize.X, LocalOutputSize.Y,
-			0., 0., OutputResolution.X, OutputResolution.Y,
-			OutputResolution,
-			OutputResolution,
-			PipelineState.VertexShader,
-			INDEX_NONE,
-			false,
-			DrawRectangleFlags);
-	}
-}
+#include "TextureResource.h"
 
 
 void ProcessOCIOColorSpaceTransform_RenderThread(
@@ -67,38 +28,47 @@ void ProcessOCIOColorSpaceTransform_RenderThread(
 {
 	check(IsInRenderingThread());
 
-	SCOPED_DRAW_EVENT(InRHICmdList, ProcessOCIOColorSpaceTransform);
-
-	InRHICmdList.Transition(FRHITransitionInfo(OutputSpaceColorTexture, ERHIAccess::Unknown, ERHIAccess::RTV));
+	FMemMark Mark(FMemStack::Get());
+	FRDGBuilder GraphBuilder(InRHICmdList);
 
-	FRHIRenderPassInfo RPInfo(OutputSpaceColorTexture, ERenderTargetActions::DontLoad_Store);
-	InRHICmdList.BeginRenderPass(RPInfo, TEXT("ProcessOCIOColorSpaceXfrm"));
+	FRDGTextureRef InputTexture = GraphBuilder.RegisterExternalTexture(CreateRenderTarget(InputSpaceColorTexture->GetTexture2D(), TEXT("OCIOInputTexture")));
+	FRDGTextureRef OutputTexture = GraphBuilder.RegisterExternalTexture(CreateRenderTarget(OutputSpaceColorTexture->GetTexture2D(), TEXT("OCIORenderTargetTexture")));
+	FScreenPassTextureViewport Viewport = FScreenPassTextureViewport(OutputTexture);
+	FScreenPassRenderTarget ScreenPassRenderTarget = FScreenPassRenderTarget(OutputTexture, FIntRect(FIntPoint::ZeroValue, OutputResolution), ERenderTargetLoadAction::EClear);
 
-	// Get shader from shader map.
-	FGlobalShaderMap* GlobalShaderMap = GetGlobalShaderMap(InFeatureLevel);
-	TShaderMapRef<FOpenColorIOVertexShader> VertexShader(GlobalShaderMap);
 	TShaderRef<FOpenColorIOPixelShader> OCIOPixelShader = InOCIOColorTransformResource->GetShader<FOpenColorIOPixelShader>();
 
-	FScreenPassPipelineState PipelineState(VertexShader, OCIOPixelShader, TStaticBlendState<>::GetRHI(), TStaticDepthStencilState<false, CF_Always>::GetRHI());
-	DrawScreenPass(InRHICmdList, OutputResolution, PipelineState, [&](FRHICommandListImmediate& RHICmdList)
+	FOpenColorIOPixelShaderParameters* Parameters = GraphBuilder.AllocParameters<FOpenColorIOPixelShaderParameters>();
+	Parameters->InputTexture = InputTexture;
+	Parameters->InputTextureSampler = TStaticSamplerState<>::GetRHI();
+	if (InLUT3dResource != nullptr)
 	{
-		// Set Gamma to 1., since we do not have any display parameters or requirement for Gamma.
-		const float Gamma = 1.0;
-
-		// Update pixel shader parameters.
-		OCIOPixelShader->SetParameters(InRHICmdList, InputSpaceColorTexture, Gamma);
-
-		if (InLUT3dResource != nullptr)
-		{
-			OCIOPixelShader->SetLUTParameter(InRHICmdList, InLUT3dResource);
-		}
-	});
-
-	// Resolve render target.
-	InRHICmdList.EndRenderPass();
-
-	// Restore readable state
-	InRHICmdList.Transition(FRHITransitionInfo(OutputSpaceColorTexture, ERHIAccess::RTV, ERHIAccess::SRVGraphics));
+		Parameters->Ocio_lut3d_0 = InLUT3dResource->TextureRHI;
+	}
+	Parameters->Ocio_lut3d_0Sampler = TStaticSamplerState<SF_Bilinear, AM_Clamp, AM_Clamp, AM_Clamp>::GetRHI();
+	// Set Gamma to 1., since we do not have any display parameters or requirement for Gamma.
+	Parameters->Gamma = 1.0;
+	Parameters->RenderTargets[0] = ScreenPassRenderTarget.GetRenderTargetBinding();
+
+	//Dummy ViewFamily/ViewInfo created for AddDrawScreenPass.
+	FSceneViewFamily ViewFamily(FSceneViewFamily::ConstructionValues(nullptr, nullptr, FEngineShowFlags(ESFIM_Game))
+		.SetTime(FGameTime())
+		.SetGammaCorrection(1.0f));
+	FSceneViewInitOptions ViewInitOptions;
+	ViewInitOptions.ViewFamily = &ViewFamily;
+	ViewInitOptions.SetViewRectangle(ScreenPassRenderTarget.ViewRect);
+	ViewInitOptions.ViewOrigin = FVector::ZeroVector;
+	ViewInitOptions.ViewRotationMatrix = FMatrix::Identity;
+	ViewInitOptions.ProjectionMatrix = FMatrix::Identity;
+	FViewInfo* DummyView = new(FMemStack::Get()) FViewInfo(ViewInitOptions);
+
+	AddDrawScreenPass(GraphBuilder, RDG_EVENT_NAME("ProcessOCIOColorSpaceXfrm"), *DummyView, Viewport, Viewport, OCIOPixelShader, Parameters);
+
+	GraphBuilder.UseExternalAccessMode(OutputTexture, ERHIAccess::SRVMask);
+	GraphBuilder.Execute();
+
+	// Properly clear the reference to ViewUniformBuffer before memstack wipes the memory
+	DummyView->~FViewInfo();
 }
 
 // static
diff --git a/Engine/Plugins/Compositing/OpenColorIO/Source/OpenColorIO/Private/OpenColorIORenderingPrivate.h b/Engine/Plugins/Compositing/OpenColorIO/Source/OpenColorIO/Private/OpenColorIORenderingPrivate.h
deleted file mode 100644
diff --git a/Engine/Plugins/Compositing/OpenColorIO/Source/OpenColorIO/Private/OpenColorIOShader.cpp b/Engine/Plugins/Compositing/OpenColorIO/Source/OpenColorIO/Private/OpenColorIOShader.cpp
--- a/Engine/Plugins/Compositing/OpenColorIO/Source/OpenColorIO/Private/OpenColorIOShader.cpp
+++ b/Engine/Plugins/Compositing/OpenColorIO/Source/OpenColorIO/Private/OpenColorIOShader.cpp
@@ -2,65 +2,4 @@
 
 #include "OpenColorIOShader.h"
 
-#include "Engine/VolumeTexture.h"
-#include "OpenColorIODerivedDataVersion.h"
-#include "OpenColorIOShaderCompilationManager.h"
-#include "OpenColorIOShared.h"
-#include "ProfilingDebugging/CookStats.h"
-#include "ProfilingDebugging/DiagnosticTable.h"
-#include "Serialization/MemoryReader.h"
-#include "Serialization/MemoryWriter.h"
-#include "ShaderCompiler.h"
-#include "ShaderParameterUtils.h"
-#include "Stats/StatsMisc.h"
-#include "UObject/UObjectGlobals.h"
-
-IMPLEMENT_TYPE_LAYOUT(FOpenColorIOShader);
-
-IMPLEMENT_SHADER_TYPE(, FOpenColorIOVertexShader, TEXT("/Plugin/OpenColorIO/Private/OpenColorIOBaseVS.usf"), TEXT("MainVS"), SF_Vertex)
-
-/** 
-* The following two shaders are identical (on hlsl side) except for the way the resources are bound on CPU side. 
-* One is oriented for RHI, the other one for RDG. 
-*/
 IMPLEMENT_SHADER_TYPE(, FOpenColorIOPixelShader, TEXT("/Plugin/OpenColorIO/Private/OpenColorIOShader.usf"), TEXT("MainPS"), SF_Pixel)
-IMPLEMENT_SHADER_TYPE(, FOpenColorIOPixelShader_RDG, TEXT("/Plugin/OpenColorIO/Private/OpenColorIOShader.usf"), TEXT("MainPS"), SF_Pixel)
-
-
-
-//////////////////////////////////////////////////////////////////////////
-
-
-FOpenColorIOPixelShader::FOpenColorIOPixelShader(const FOpenColorIOShaderType::CompiledShaderInitializerType& Initializer)
-	: FOpenColorIOShader(Initializer)
-	, DebugDescription(Initializer.DebugDescription)
-{
-	BindParams(Initializer.ParameterMap);
-}
-
-void FOpenColorIOPixelShader::SetParameters(FRHICommandList& InRHICmdList, FTextureResource* InInputTexture, float InGamma)
-{
-	SetParameters(InRHICmdList, InInputTexture->TextureRHI, InGamma);
-}
-
-void FOpenColorIOPixelShader::SetParameters(FRHICommandList& InRHICmdList, FTextureRHIRef InInputTexture, float InGamma)
-{
-	SetTextureParameter(InRHICmdList, InRHICmdList.GetBoundPixelShader(), InputTexture, InputTextureSampler, TStaticSamplerState<SF_Bilinear, AM_Clamp, AM_Clamp, AM_Clamp>::GetRHI(), InInputTexture);
-	SetShaderValue(InRHICmdList, InRHICmdList.GetBoundPixelShader(), Gamma, InGamma);
-}
-
-void FOpenColorIOPixelShader::SetLUTParameter(FRHICommandList& InRHICmdList, FTextureResource* InLUT3dResource)
-{
-	SetTextureParameter(InRHICmdList, InRHICmdList.GetBoundPixelShader(), OCIO3dTexture, OCIO3dTextureSampler, TStaticSamplerState<SF_Bilinear, AM_Clamp, AM_Clamp, AM_Clamp>::GetRHI(), InLUT3dResource->TextureRHI);
-}
-
-void FOpenColorIOPixelShader::BindParams(const FShaderParameterMap& ParameterMap)
-{
-	InputTexture.Bind(ParameterMap, TEXT("InputTexture"));
-	InputTextureSampler.Bind(ParameterMap, TEXT("InputTextureSampler"));
-
-	OCIO3dTexture.Bind(ParameterMap, OpenColorIOShader::OCIOLut3dName);
-	OCIO3dTextureSampler.Bind(ParameterMap, OpenColorIOShader::OCIOLut3dSamplerName);
-
-	Gamma.Bind(ParameterMap, TEXT("Gamma"));
-}
diff --git a/Engine/Plugins/Compositing/OpenColorIO/Source/OpenColorIO/Private/OpenColorIOShaderMap.cpp b/Engine/Plugins/Compositing/OpenColorIO/Source/OpenColorIO/Private/OpenColorIOShaderMap.cpp
--- a/Engine/Plugins/Compositing/OpenColorIO/Source/OpenColorIO/Private/OpenColorIOShaderMap.cpp
+++ b/Engine/Plugins/Compositing/OpenColorIO/Source/OpenColorIO/Private/OpenColorIOShaderMap.cpp
@@ -479,7 +479,6 @@ FShader* FOpenColorIOShaderMap::ProcessCompilationResultsForSingleJob(const TRef
 	Shader = OpenColorIOShaderType->FinishCompileShader(InShaderMapHash, *CurrentJob, GetContent()->FriendlyName);
 	bCompiledSuccessfully = CurrentJob->bSucceeded;
 
-	FOpenColorIOPixelShader* OpenColorIOShader = static_cast<FOpenColorIOPixelShader*>(Shader);
 	check(Shader && Shader->GetCodeSize() > 0);
 	check(!GetContent()->HasShader(OpenColorIOShaderType, /* PermutationId = */ 0));
 	return GetMutableContent()->FindOrAddShader(OpenColorIOShaderType->GetHashedName(), CurrentJob->Key.PermutationId, Shader);
diff --git a/Engine/Plugins/Compositing/OpenColorIO/Source/OpenColorIO/Public/OpenColorIOShader.h b/Engine/Plugins/Compositing/OpenColorIO/Source/OpenColorIO/Public/OpenColorIOShader.h
--- a/Engine/Plugins/Compositing/OpenColorIO/Source/OpenColorIO/Public/OpenColorIOShader.h
+++ b/Engine/Plugins/Compositing/OpenColorIO/Source/OpenColorIO/Public/OpenColorIOShader.h
@@ -9,15 +9,12 @@
 #include "GlobalShader.h"
 #include "OpenColorIOShaderType.h"
 #include "OpenColorIOShared.h"
-#include "SceneView.h"
-#include "Shader.h"
 #include "ShaderParameters.h"
-#include "GlobalShader.h"
 #include "ShaderParameterStruct.h"
+
 class FTextureResource;
 class UClass;
 
-/** RDG Parameters to be used with FOpenColorIOPixelShader_RDG. */
 BEGIN_SHADER_PARAMETER_STRUCT(FOpenColorIOPixelShaderParameters, )
 	SHADER_PARAMETER_RDG_TEXTURE(Texture2D, InputTexture)
 	SHADER_PARAMETER_SAMPLER(SamplerState, InputTextureSampler)
@@ -30,95 +27,17 @@ END_SHADER_PARAMETER_STRUCT()
 namespace OpenColorIOShader
 {
 	static const TCHAR* OpenColorIOShaderFunctionName = TEXT("OCIOConvert");
-	static const TCHAR* OCIOLut3dName = TEXT("Ocio_lut3d_0");
-	static const TCHAR* OCIOLut3dSamplerName = TEXT("Ocio_lut3d_0Sampler");
-	static const uint32 MaximumTextureNumber = 10;
 	static const uint32 Lut3dEdgeLength = 65;
 }
 
-class FOpenColorIOShader : public FShader
-{
-	DECLARE_TYPE_LAYOUT(FOpenColorIOShader, NonVirtual);
-public:
-	FOpenColorIOShader() = default;
-	FOpenColorIOShader(const FOpenColorIOShader::CompiledShaderInitializerType& Initializer) : FShader(Initializer) {}
-};
-
-/** Vertex shader compatible with both RHI and RDG. */
-class OPENCOLORIO_API FOpenColorIOVertexShader : public FGlobalShader
+class OPENCOLORIO_API FOpenColorIOPixelShader : public FGlobalShader
 {
-	DECLARE_SHADER_TYPE(FOpenColorIOVertexShader, Global);
-
 public:
-	static bool ShouldCompilePermutation(const FGlobalShaderPermutationParameters& Parameters)
-	{
-		return IsFeatureLevelSupported(Parameters.Platform, ERHIFeatureLevel::ES3_1);
-	}
-
-	static void ModifyCompilationEnvironment(const FGlobalShaderPermutationParameters& Parameters, FShaderCompilerEnvironment& OutEnvironment)
-	{
-		FGlobalShader::ModifyCompilationEnvironment(Parameters, OutEnvironment);
-	}
-
-	/** Default constructor. */
-	FOpenColorIOVertexShader() {}
-
-	/** Initialization constructor. */
-	FOpenColorIOVertexShader(const ShaderMetaType::CompiledShaderInitializerType& Initializer)
-		: FGlobalShader(Initializer)
-	{
-	}
-};
-
-
-/** Base class of all shaders that need OpenColorIO pixel shader parameters. This Pixel shader is to be used only within RHI. */
-class OPENCOLORIO_API FOpenColorIOPixelShader : public FOpenColorIOShader
-{
 	DECLARE_SHADER_TYPE(FOpenColorIOPixelShader, OpenColorIO);
+	SHADER_USE_PARAMETER_STRUCT(FOpenColorIOPixelShader, FGlobalShader);
 
-public:
-	using FPermutationParameters = FOpenColorIOShaderPermutationParameters;
-
-	FOpenColorIOPixelShader() = default;
-	FOpenColorIOPixelShader(const FOpenColorIOShaderType::CompiledShaderInitializerType& Initializer);
-
-	static bool ShouldCompilePermutation(const FOpenColorIOShaderPermutationParameters& Parameters)
-	{
-		return true;
-	}
-
-
-	void SetParameters(FRHICommandList& InRHICmdList, FTextureResource* InInputTexture, float InGamma);
-	void SetParameters(FRHICommandList& InRHICmdList, FTextureRHIRef InInputTexture, float InGamma);
-	void SetLUTParameter(FRHICommandList& InRHICmdList, FTextureResource* InLUT3dResource);
-
-	// Bind parameters
-	void BindParams(const FShaderParameterMap& ParameterMap);
-
-protected:
-	LAYOUT_FIELD(FShaderResourceParameter, InputTexture)
-	LAYOUT_FIELD(FShaderResourceParameter, InputTextureSampler)
-	LAYOUT_FIELD(FShaderParameter, Gamma)
-
-	LAYOUT_FIELD(FShaderResourceParameter, OCIO3dTexture)
-	LAYOUT_FIELD(FShaderResourceParameter, OCIO3dTextureSampler)
-
-private:
-	LAYOUT_FIELD(FMemoryImageString, DebugDescription)
-};
-
-
-/** 
-*   Pixel shader to be used within RDG environment. 
-*	It is identical (on hlsl side) to the above FOpenColorIOPixelShader except for the way the resources are bound on CPU side. 
-*/
-class OPENCOLORIO_API FOpenColorIOPixelShader_RDG : public FGlobalShader
-{
-public:
 	using FParameters = FOpenColorIOPixelShaderParameters;
 	using FPermutationParameters = FOpenColorIOShaderPermutationParameters;
-	DECLARE_SHADER_TYPE(FOpenColorIOPixelShader_RDG, OpenColorIO);
-	SHADER_USE_PARAMETER_STRUCT_WITH_LEGACY_BASE(FOpenColorIOPixelShader_RDG, FGlobalShader);
 
 	static bool ShouldCompilePermutation(const FOpenColorIOShaderPermutationParameters& Parameters)
 	{
diff --git a/Engine/Plugins/Compositing/OpenColorIO/Source/OpenColorIO/Public/OpenColorIOShared.h b/Engine/Plugins/Compositing/OpenColorIO/Source/OpenColorIO/Public/OpenColorIOShared.h
--- a/Engine/Plugins/Compositing/OpenColorIO/Source/OpenColorIO/Public/OpenColorIOShared.h
+++ b/Engine/Plugins/Compositing/OpenColorIO/Source/OpenColorIO/Public/OpenColorIOShared.h
@@ -24,7 +24,6 @@
 class FOpenColorIOTransformResource;
 class FOpenColorIOShaderMap;
 class FOpenColorIOPixelShader;
-class FOpenColorIOPixelShader_RDG;
 class FOpenColorIOShaderMapId;
 
 /** Stores outputs from the color transform compile that need to be saved. */
