OCIO: Removing moot/incorrect manual transition to (default) SRVMask.

diff --git a/Engine/Plugins/Compositing/OpenColorIO/Source/OpenColorIO/Private/OpenColorIORendering.cpp b/Engine/Plugins/Compositing/OpenColorIO/Source/OpenColorIO/Private/OpenColorIORendering.cpp
--- a/Engine/Plugins/Compositing/OpenColorIO/Source/OpenColorIO/Private/OpenColorIORendering.cpp
+++ b/Engine/Plugins/Compositing/OpenColorIO/Source/OpenColorIO/Private/OpenColorIORendering.cpp
@@ -64,7 +64,6 @@ void ProcessOCIOColorSpaceTransform_RenderThread(
 
 	AddDrawScreenPass(GraphBuilder, RDG_EVENT_NAME("ProcessOCIOColorSpaceXfrm"), *DummyView, Viewport, Viewport, OCIOPixelShader, Parameters);
 
-	GraphBuilder.UseExternalAccessMode(OutputTexture, ERHIAccess::SRVMask);
 	GraphBuilder.Execute();
 
 	// Properly clear the reference to ViewUniformBuffer before memstack wipes the memory
