OCIO: Migrating EditorPerProjectUserSettings config to plugin-specific file. Settings in the old location will be loaded and migrated to the new config file (with warning).

diff --git a/Engine/Plugins/Compositing/OpenColorIO/Source/OpenColorIOEditor/Private/OpenColorIOEditorSettings.cpp b/Engine/Plugins/Compositing/OpenColorIO/Source/OpenColorIOEditor/Private/OpenColorIOEditorSettings.cpp
--- a/Engine/Plugins/Compositing/OpenColorIO/Source/OpenColorIOEditor/Private/OpenColorIOEditorSettings.cpp
+++ b/Engine/Plugins/Compositing/OpenColorIO/Source/OpenColorIOEditor/Private/OpenColorIOEditorSettings.cpp
@@ -2,6 +2,9 @@
 
 #include "OpenColorIOEditorSettings.h"
 
+#include "IOpenColorIOEditorModule.h"
+#include "Misc/ConfigCacheIni.h"
+#include "Misc/Paths.h"
 
 const FOpenColorIODisplayConfiguration* UOpenColorIOLevelViewportSettings::GetViewportSettings(FName ViewportIdentifier) const
 {
@@ -18,6 +21,27 @@ const FOpenColorIODisplayConfiguration* UOpenColorIOLevelViewportSettings::GetVi
 	return nullptr;
 }
 
+void UOpenColorIOLevelViewportSettings::PostInitProperties()
+{
+	Super::PostInitProperties();
+
+	if (GConfig && FPaths::FileExists(GEditorPerProjectIni))
+	{
+		const TCHAR* SectionName = TEXT("/Script/OpenColorIOEditor.OpenColorIOLevelViewportSettings");
+
+		if (GConfig->DoesSectionExist(SectionName, GEditorPerProjectIni))
+		{
+			LoadConfig(UOpenColorIOLevelViewportSettings::StaticClass(), *GEditorPerProjectIni);
+
+			SaveConfig();
+
+			GConfig->EmptySection(SectionName, *GEditorPerProjectIni);
+
+			UE_LOG(LogOpenColorIOEditor, Warning, TEXT("Migrated EditorPerProjectUserSettings OpenColorIO settings to plugin-specific config file: %s."), *GetClass()->GetConfigName());
+		}
+	}
+}
+
 void UOpenColorIOLevelViewportSettings::SetViewportSettings(FName ViewportIdentifier, const FOpenColorIODisplayConfiguration& Configuration)
 {
 	FPerViewportDisplaySettingPair* Pair = ViewportsSettings.FindByPredicate([ViewportIdentifier](const FPerViewportDisplaySettingPair& Other)
diff --git a/Engine/Plugins/Compositing/OpenColorIO/Source/OpenColorIOEditor/Private/OpenColorIOEditorSettings.h b/Engine/Plugins/Compositing/OpenColorIO/Source/OpenColorIOEditor/Private/OpenColorIOEditorSettings.h
--- a/Engine/Plugins/Compositing/OpenColorIO/Source/OpenColorIOEditor/Private/OpenColorIOEditorSettings.h
+++ b/Engine/Plugins/Compositing/OpenColorIO/Source/OpenColorIOEditor/Private/OpenColorIOEditorSettings.h
@@ -32,13 +32,16 @@ struct FPerViewportDisplaySettingPair
 /**
  * List of settings associated to level viewport instances linked with an identifier
  */
-UCLASS(config = EditorPerProjectUserSettings)
+
+UCLASS(config = OpenColorIO)
 class UOpenColorIOLevelViewportSettings : public UObject
 {
 	GENERATED_BODY()
 
 public:
 
+	virtual void PostInitProperties() override;
+
 	/** Returns setting associated with a given viewport identifier */
 	const FOpenColorIODisplayConfiguration* GetViewportSettings(FName ViewportIdentifier) const;
 
