OCIO: Adding support for OCIOv2's gpu processor, while maintaining legacy processor support via a plugin setting.

* Dynamic property support remains disabled still.
* Prior to further improvements to shader bindings, a maximum of 8 LUTs per transform are supported.

diff --git a/Engine/Plugins/Compositing/OpenColorIO/Source/OpenColorIO/OpenColorIO.Build.cs b/Engine/Plugins/Compositing/OpenColorIO/Source/OpenColorIO/OpenColorIO.Build.cs
--- a/Engine/Plugins/Compositing/OpenColorIO/Source/OpenColorIO/OpenColorIO.Build.cs
+++ b/Engine/Plugins/Compositing/OpenColorIO/Source/OpenColorIO/OpenColorIO.Build.cs
@@ -18,9 +18,6 @@ namespace UnrealBuildTool.Rules
 
 			PrivateDependencyModuleNames.AddRange(
 				new string[] {
-					"Core",
-					"CoreUObject",
-					"Engine",
 					"Projects",
 					"RenderCore",
 					"RHI",
@@ -43,6 +40,10 @@ namespace UnrealBuildTool.Rules
 
 			PublicDependencyModuleNames.AddRange(
 				new string[] {
+					"Core",
+					"CoreUObject",
+					"Engine",
+					"DeveloperSettings"
 				});
 
 			if (Target.bBuildEditor == true)
diff --git a/Engine/Plugins/Compositing/OpenColorIO/Source/OpenColorIO/Private/OpenColorIOColorTransform.cpp b/Engine/Plugins/Compositing/OpenColorIO/Source/OpenColorIO/Private/OpenColorIOColorTransform.cpp
--- a/Engine/Plugins/Compositing/OpenColorIO/Source/OpenColorIO/Private/OpenColorIOColorTransform.cpp
+++ b/Engine/Plugins/Compositing/OpenColorIO/Source/OpenColorIO/Private/OpenColorIOColorTransform.cpp
@@ -7,6 +7,7 @@
 #include "Modules/ModuleManager.h"
 #include "OpenColorIOConfiguration.h"
 #include "OpenColorIOModule.h"
+#include "OpenColorIOSettings.h"
 #include "UObject/UObjectIterator.h"
 
 
@@ -24,7 +25,20 @@
 
 #endif //WITH_EDITOR
 
+namespace {
+#if WITH_EDITOR && WITH_OCIO
+	/*
+	 * Get the processor optimization flag.
+	 * @todo: Remove "no dynamic properties" once this path is enabled.
+	 */
+	OCIO_NAMESPACE::OptimizationFlags GetProcessorOptimization()
+	{
+		using namespace OCIO_NAMESPACE;
 
+		return static_cast<OptimizationFlags>(OptimizationFlags::OPTIMIZATION_DEFAULT | OptimizationFlags::OPTIMIZATION_NO_DYNAMIC_PROPERTIES);
+	}
+#endif
+}
 
 
 void UOpenColorIOColorTransform::SerializeOpenColorIOShaderMaps(const TMap<const ITargetPlatform*, TArray<FOpenColorIOTransformResource*>>* PlatformColorTransformResourcesToSavePtr, FArchive& Ar, TArray<FOpenColorIOTransformResource>&  OutLoadedResources)
@@ -95,10 +109,15 @@ void UOpenColorIOColorTransform::ProcessSerializedShaderMaps(UOpenColorIOColorTr
 	}
 }
 
-void UOpenColorIOColorTransform::GetOpenColorIOLUTKeyGuid(const FString& InLutIdentifier, FGuid& OutLutGuid)
+void UOpenColorIOColorTransform::GetOpenColorIOLUTKeyGuid(const FString& InProcessorIdentifier, FGuid& OutLutGuid)
+{
+	GetOpenColorIOLUTKeyGuid(InProcessorIdentifier, FName(), OutLutGuid);
+}
+
+void UOpenColorIOColorTransform::GetOpenColorIOLUTKeyGuid(const FString& InProcessorIdentifier, const FName& InName, FGuid& OutLutGuid)
 {
 #if WITH_EDITOR
-	FString DDCKey = FDerivedDataCacheInterface::BuildCacheKey(TEXT("OCIOLUT"), OPENCOLORIO_DERIVEDDATA_VER, *InLutIdentifier);
+	FString DDCKey = FDerivedDataCacheInterface::BuildCacheKey(TEXT("OCIOLUT"), OPENCOLORIO_DERIVEDDATA_VER, *InProcessorIdentifier);
 
 #if WITH_OCIO
 	//Keep library version in the DDC key to invalidate it once we move to a new library
@@ -106,6 +125,11 @@ void UOpenColorIOColorTransform::GetOpenColorIOLUTKeyGuid(const FString& InLutId
 	DDCKey += TEXT(OCIO_VERSION);
 #endif //WITH_OCIO
 
+	if (!InName.IsNone())
+	{
+		DDCKey += InName.ToString();
+	}
+
 	const uint32 KeyLength = DDCKey.Len() * sizeof(DDCKey[0]);
 	uint32 Hash[5];
 	FSHA1::HashBuffer(*DDCKey, KeyLength, reinterpret_cast<uint8*>(Hash));
@@ -166,17 +190,14 @@ void UOpenColorIOColorTransform::SerializeLuts(FArchive& Ar)
 		int32 Num3dLutsToSave = 0;
 		if (Ar.IsCooking())
 		{
-			if (Lut3dTexture != nullptr)
-			{
-				Num3dLutsToSave = 1;
-			}
+			Num3dLutsToSave = Textures.Num();
 		}
 
 		Ar << Num3dLutsToSave;
 
 		if (Num3dLutsToSave > 0)
 		{
-			Ar << Lut3dTexture;
+			Ar << Textures;
 		}
 	}
 	else if (Ar.IsLoading())
@@ -187,14 +208,14 @@ void UOpenColorIOColorTransform::SerializeLuts(FArchive& Ar)
 		if (NumLoaded3dLuts > 0)
 		{
 			//Will only happen on cooked data
-			Ar << Lut3dTexture;
+			Ar << Textures;
 		}
 	}
 }
 
 void UOpenColorIOColorTransform::CacheResourceTextures()
 {
-	if (Lut3dTexture == nullptr)
+	if (Textures.IsEmpty())
 	{
 #if WITH_EDITOR && WITH_OCIO
 		OCIO_NAMESPACE::ConstConfigRcPtr CurrentConfig = ConfigurationOwner->GetLoadedConfigurationFile();
@@ -205,6 +226,7 @@ void UOpenColorIOColorTransform::CacheResourceTextures()
 #endif
 			{
 				OCIO_NAMESPACE::ConstProcessorRcPtr TransformProcessor = CurrentConfig->getProcessor(StringCast<ANSICHAR>(*SourceColorSpace).Get(), StringCast<ANSICHAR>(*DestinationColorSpace).Get());
+				
 				if (TransformProcessor)
 				{
 					OCIO_NAMESPACE::GpuShaderDescRcPtr ShaderDescription = OCIO_NAMESPACE::GpuShaderDesc::CreateShaderDesc();
@@ -212,27 +234,87 @@ void UOpenColorIOColorTransform::CacheResourceTextures()
 					ShaderDescription->setFunctionName(StringCast<ANSICHAR>(OpenColorIOShader::OpenColorIOShaderFunctionName).Get());
 					ShaderDescription->setResourcePrefix("Ocio");
 
-					OCIO_NAMESPACE::ConstGPUProcessorRcPtr GPUProcessor = TransformProcessor->getOptimizedLegacyGPUProcessor(OCIO_NAMESPACE::OptimizationFlags::OPTIMIZATION_DEFAULT, OpenColorIOShader::Lut3dEdgeLength);
+					OCIO_NAMESPACE::ConstGPUProcessorRcPtr GPUProcessor = nullptr;
+					const UOpenColorIOSettings* Settings = GetDefault<UOpenColorIOSettings>();
+					check(Settings);
+
+					if (Settings->bUseLegacyProcessor)
+					{
+						unsigned int EdgeLength = static_cast<unsigned int>(OpenColorIOShader::Lut3dEdgeLength);
+						GPUProcessor = TransformProcessor->getOptimizedLegacyGPUProcessor(GetProcessorOptimization(), EdgeLength);
+					}
+					else
+					{
+						GPUProcessor = TransformProcessor->getOptimizedGPUProcessor(GetProcessorOptimization());
+					}
+
 					GPUProcessor->extractGpuShaderInfo(ShaderDescription);
+					// @todo: Remove once we add support for dynamic properties
+					ensureMsgf(ShaderDescription->getNumDynamicProperties() == 0, TEXT("Dynamic properties are not currently supported."));
+
+					const FString ProcessorID = StringCast<TCHAR>(GPUProcessor->getCacheID()).Get();
 
-					FString Lut3dIdentifier = StringCast<TCHAR>(GPUProcessor->getCacheID()).Get();
-					if (Lut3dIdentifier != TEXT("<NULL>") && ShaderDescription->getNum3DTextures() > 0 )
+					//In editor, it will use what's on DDC if there's something corresponding to the actual data or use that raw data
+					//that OCIO library has on board. The textures will be serialized only when cooking.
+					
+					// Process 3D luts
+					for (uint32 Index = 0; Index < ShaderDescription->getNum3DTextures(); ++Index)
 					{
 						const char* TextureName = nullptr;
 						const char* SamplerName = nullptr;
-						unsigned int EdgeLength = static_cast<unsigned int>(OpenColorIOShader::Lut3dEdgeLength);
-						OCIO_NAMESPACE::Interpolation Interpolation = OCIO_NAMESPACE::INTERP_BEST;
-						ShaderDescription->get3DTexture(0, TextureName, SamplerName, EdgeLength, Interpolation);
-						checkf(TextureName && *TextureName && SamplerName && *SamplerName && EdgeLength > 0, TEXT("Invalid OCIO texture or sampler."));
+						unsigned int EdgeLength;
+						OCIO_NAMESPACE::Interpolation Interpolation = OCIO_NAMESPACE::INTERP_TETRAHEDRAL;
+
+						ShaderDescription->get3DTexture(Index, TextureName, SamplerName, EdgeLength, Interpolation);
+						checkf(TextureName && *TextureName && SamplerName && *SamplerName && EdgeLength > 0, TEXT("Invalid OCIO 3D texture or sampler."));
+
+						const float* TextureValues = 0x0;
+						ShaderDescription->get3DTextureValues(Index, TextureValues);
+						checkf(TextureValues, TEXT("Failed to read OCIO 3D LUT data."));
 
-						const float* Lut3dData = 0x0;
-						ShaderDescription->get3DTextureValues(0, Lut3dData);
-						checkf(Lut3dData, TEXT("Failed to read OCIO 3d LUT data."));
+						TextureFilter Filter = TF_Bilinear;
+						if (Interpolation == OCIO_NAMESPACE::Interpolation::INTERP_NEAREST || Interpolation == OCIO_NAMESPACE::Interpolation::INTERP_TETRAHEDRAL)
+						{
+							Filter = TF_Nearest;
+						}
+						TObjectPtr<UTexture> Result = CreateTexture3DLUT(ProcessorID, FName(TextureName), EdgeLength, Filter, TextureValues);
 
-						//In editor, it will use what's on DDC if there's something corresponding to the actual data or use that raw data
-						//that OCIO library has on board. The texture will be serialized only when cooking.
-						Update3dLutTexture(Lut3dIdentifier, Lut3dData);
+						Textures.Add(MoveTemp(Result));
 					}
+
+					// Process 1D luts
+					for (uint32 Index = 0; Index < ShaderDescription->getNumTextures(); ++Index)
+					{
+						const char* TextureName = nullptr;
+						const char* SamplerName = nullptr;
+						unsigned TextureWidth = 0;
+						unsigned TextureHeight = 0;
+						OCIO_NAMESPACE::GpuShaderDesc::TextureType Channel = OCIO_NAMESPACE::GpuShaderDesc::TEXTURE_RGB_CHANNEL;
+						OCIO_NAMESPACE::Interpolation Interpolation = OCIO_NAMESPACE::Interpolation::INTERP_LINEAR;
+
+						ShaderDescription->getTexture(Index, TextureName, SamplerName, TextureWidth, TextureHeight, Channel, Interpolation);
+						checkf(TextureName && *TextureName && SamplerName && *SamplerName && TextureWidth > 0, TEXT("Invalid OCIO 1D texture or sampler."));
+
+						const float* TextureValues = 0x0;
+						ShaderDescription->getTextureValues(Index, TextureValues);
+						checkf(TextureValues, TEXT("Failed to read OCIO 1D LUT data."));
+
+						TextureFilter Filter = Interpolation == OCIO_NAMESPACE::Interpolation::INTERP_NEAREST ? TF_Nearest : TF_Bilinear;
+						bool bRedChannelOnly = Channel == OCIO_NAMESPACE::GpuShaderCreator::TEXTURE_RED_CHANNEL;
+						TObjectPtr<UTexture> Result = CreateTexture1DLUT(ProcessorID, FName(TextureName), TextureWidth, TextureHeight, Filter, bRedChannelOnly, TextureValues);
+
+						Textures.Add(MoveTemp(Result));
+					}
+
+					ensureAlwaysMsgf(Textures.Num() <= (int32)OpenColorIOShader::MaximumTextureSlots, TEXT("Color transform %s exceeds our current limit of %u texture slots. Use the legacy processor instead."), *GetTransformFriendlyName(), OpenColorIOShader::MaximumTextureSlots);
+
+					// Subscript numbers extracted by FName can be used to sort textures for matching the shader binding accesses.
+					Textures.Sort(
+						[](const UTexture& TextureA, const UTexture& TextureB)
+						{
+							return TextureA.GetFName().GetNumber() < TextureB.GetFName().GetNumber();
+						}
+					);
 				}
 				else
 				{
@@ -323,15 +405,16 @@ FOpenColorIOTransformResource* UOpenColorIOColorTransform::AllocateResource()
 	return new FOpenColorIOTransformResource();
 }
 
-bool UOpenColorIOColorTransform::GetShaderAndLUTResouces(ERHIFeatureLevel::Type InFeatureLevel, FOpenColorIOTransformResource*& OutShaderResource, FTextureResource*& OutLUT3dResource)
+bool UOpenColorIOColorTransform::GetRenderResources(ERHIFeatureLevel::Type InFeatureLevel, FOpenColorIOTransformResource*& OutShaderResource, TArray<FTextureResource*>& OutTextureResources)
 {
 	OutShaderResource = ColorTransformResources[InFeatureLevel];
 	if (OutShaderResource)
 	{
-		//Some color transform will only require shader code with no LUT involved.
-		if (Lut3dTexture != nullptr)
+		OutTextureResources.Reserve(Textures.Num());
+
+		for (const TObjectPtr<UTexture>& Texture : Textures)
 		{
-			OutLUT3dResource = Lut3dTexture->GetResource();
+			OutTextureResources.Add(Texture->GetResource());
 		}
 
 		return true;
@@ -343,6 +426,18 @@ bool UOpenColorIOColorTransform::GetShaderAndLUTResouces(ERHIFeatureLevel::Type
 	}
 }
 
+bool UOpenColorIOColorTransform::GetShaderAndLUTResouces(ERHIFeatureLevel::Type InFeatureLevel, FOpenColorIOTransformResource*& OutShaderResource, FTextureResource*& OutLUT3dResource)
+{
+	TArray<FTextureResource*> OutTextureResources;
+	bool bResult = GetRenderResources(InFeatureLevel, OutShaderResource, OutTextureResources);
+
+	if (!OutTextureResources.IsEmpty())
+	{
+		OutLUT3dResource = OutTextureResources[0];
+	}
+	return bResult;
+}
+
 bool UOpenColorIOColorTransform::IsTransform(const FString& InSourceColorSpace, const FString& InDestinationColorSpace) const
 {
 	return SourceColorSpace == InSourceColorSpace && DestinationColorSpace == InDestinationColorSpace;
@@ -401,8 +496,22 @@ bool UOpenColorIOColorTransform::UpdateShaderInfo(FString& OutShaderCodeHash, FS
 				ShaderDescription->setFunctionName(StringCast<ANSICHAR>(OpenColorIOShader::OpenColorIOShaderFunctionName).Get());
 				ShaderDescription->setResourcePrefix("Ocio");
 
-				OCIO_NAMESPACE::ConstGPUProcessorRcPtr GPUProcessor = TransformProcessor->getOptimizedLegacyGPUProcessor(OCIO_NAMESPACE::OptimizationFlags::OPTIMIZATION_DEFAULT, OpenColorIOShader::Lut3dEdgeLength);
+				OCIO_NAMESPACE::ConstGPUProcessorRcPtr GPUProcessor = nullptr;
+				const UOpenColorIOSettings* Settings = GetDefault<UOpenColorIOSettings>();
+				check(Settings);
+
+				if (Settings->bUseLegacyProcessor)
+				{
+					unsigned int EdgeLength = static_cast<unsigned int>(OpenColorIOShader::Lut3dEdgeLength);
+					GPUProcessor = TransformProcessor->getOptimizedLegacyGPUProcessor(GetProcessorOptimization(), EdgeLength);
+				}
+				else
+				{
+					GPUProcessor = TransformProcessor->getOptimizedGPUProcessor(GetProcessorOptimization());
+				}
 				GPUProcessor->extractGpuShaderInfo(ShaderDescription);
+				// @todo: Remove once we add support for dynamic properties
+				ensureMsgf(ShaderDescription->getNumDynamicProperties() == 0, TEXT("We do not currently support dynamic properties."));
 
 				FString GLSLShaderCode = StringCast<TCHAR>(ShaderDescription->getShaderText()).Get();
 
@@ -444,43 +553,119 @@ bool UOpenColorIOColorTransform::UpdateShaderInfo(FString& OutShaderCodeHash, FS
 #endif
 }
 
-void UOpenColorIOColorTransform::Update3dLutTexture(const FString& InLutIdentifier, const float* InSourceData)
+TObjectPtr<UTexture> UOpenColorIOColorTransform::CreateTexture3DLUT(const FString& InProcessorIdentifier, const FName& InName, uint32 InLutLength, TextureFilter InFilter, const float* InSourceData)
 {
+	TObjectPtr<UVolumeTexture> OutTexture = nullptr;
+
 #if WITH_EDITOR && WITH_OCIO
 	check(InSourceData);
 
-	Lut3dTexture = NewObject<UVolumeTexture>(this, NAME_None, RF_NoFlags);
+	/* Note here that while it is possible to create proper 32f textures using [UTexture]::CreateTransient and reparenting
+	 * via Texture->Rename(nullptr, this), cooking would fail as it remains unsupported currently for those formats.
+	 * (The same note applies to 1D LUT creation.) */
+
+	OutTexture = NewObject<UVolumeTexture>(this, InName);
 
 	//Initializes source data with the raw LUT. If it's found in DDC, the resulting platform data will be fetched from there. 
 	//If not, the source data will be used to generate the platform data.
-	Lut3dTexture->MipGenSettings = TMGS_NoMipmaps;
-	Lut3dTexture->CompressionNone = true;
-	Lut3dTexture->Source.Init(OpenColorIOShader::Lut3dEdgeLength, OpenColorIOShader::Lut3dEdgeLength, OpenColorIOShader::Lut3dEdgeLength, /*NumMips=*/ 1, TSF_RGBA16F, nullptr);
+	OutTexture->MipGenSettings = TMGS_NoMipmaps;
+	OutTexture->SRGB = 0;
+	OutTexture->CompressionNone = true;
+	OutTexture->Filter = InFilter;
+	OutTexture->AddressMode = TextureAddress::TA_Clamp;
+	OutTexture->Source.Init(InLutLength, InLutLength, InLutLength, /*NumMips=*/ 1, TSF_RGBA16F, nullptr);
+
+	FFloat16Color* MipData = reinterpret_cast<FFloat16Color*>(OutTexture->Source.LockMip(0));
+	for (uint32 Z = 0; Z < InLutLength; ++Z)
+	{
+		for (uint32 Y = 0; Y < InLutLength; Y++)
+		{
+			FFloat16Color* Row = &MipData[Y * InLutLength + Z * InLutLength * InLutLength];
+			const float* Source = &InSourceData[Y * InLutLength * 3 + Z * InLutLength * InLutLength * 3];
+			for (uint32 X = 0; X < InLutLength; X++)
+			{
+				FFloat16Color& TargetColor = Row[X];
+				TargetColor.R.SetClamped(Source[X * 3 + 0]);
+				TargetColor.G.SetClamped(Source[X * 3 + 1]);
+				TargetColor.B.SetClamped(Source[X * 3 + 2]);
+				TargetColor.A.SetOne();
+			}
+		}
+	}
+	OutTexture->Source.UnlockMip(0);
 
-	FFloat16Color* MipData = reinterpret_cast<FFloat16Color*>(Lut3dTexture->Source.LockMip(0));
-	const uint32 LutLength = OpenColorIOShader::Lut3dEdgeLength;
-	for (uint32 Z = 0; Z < LutLength; ++Z)
+	//Generate a Guid from the identifier received from the library and our DDC version.
+	FGuid LutGuid;
+	GetOpenColorIOLUTKeyGuid(InProcessorIdentifier, InName, LutGuid);
+	OutTexture->Source.SetId(LutGuid, true);
+
+	//Process our new texture to be usable in rendering pipeline.
+	OutTexture->UpdateResource();
+#endif
+
+	return OutTexture;
+}
+
+void UOpenColorIOColorTransform::Update3dLutTexture(const FString& InLutIdentifier, const float* InSourceData)
+{
+#if WITH_EDITOR && WITH_OCIO
+	Textures.Empty();
+
+	TObjectPtr<UTexture> Texture = CreateTexture3DLUT(InLutIdentifier, FName(), OpenColorIOShader::Lut3dEdgeLength, TextureFilter::TF_Bilinear, InSourceData);
+	Textures.Add(MoveTemp(Texture));
+#endif
+}
+
+TObjectPtr<UTexture> UOpenColorIOColorTransform::CreateTexture1DLUT(const FString& InProcessorIdentifier, const FName& InName, uint32 InTextureWidth, uint32 InTextureHeight, TextureFilter InFilter, bool bRedChannelOnly, const float* InSourceData)
+{
+	TObjectPtr<UTexture2D> OutTexture = nullptr;
+
+#if WITH_EDITOR && WITH_OCIO
+	check(InSourceData);
+
+	OutTexture = NewObject<UTexture2D>(this, InName);
+	OutTexture->MipGenSettings = TMGS_NoMipmaps;
+	OutTexture->SRGB = 0;
+	OutTexture->CompressionNone = true;
+	OutTexture->Filter = InFilter;
+	OutTexture->AddressX = TextureAddress::TA_Clamp;
+	OutTexture->AddressY = TextureAddress::TA_Clamp;
+
+	if (bRedChannelOnly)
 	{
-		for (uint32 Y = 0; Y < LutLength; Y++)
+		OutTexture->Source.Init(InTextureWidth, InTextureHeight, /*NumSlices=*/ 1, /*NumMips=*/ 1, TSF_R32F, reinterpret_cast<const uint8*>(InSourceData));
+	}
+	else
+	{
+		OutTexture->Source.Init(InTextureWidth, InTextureHeight, /*NumSlices=*/ 1, /*NumMips=*/ 1, TSF_RGBA16F, nullptr);
+
+		FFloat16Color* MipData = reinterpret_cast<FFloat16Color*>(OutTexture->Source.LockMip(0));
+		for (uint32 Y = 0; Y < InTextureHeight; Y++)
 		{
-			FFloat16Color* Row = &MipData[Y * LutLength + Z * LutLength * LutLength];
-			const float* Source = &InSourceData[Y * LutLength * 3 + Z * LutLength * LutLength * 3];
-			for (uint32 X = 0; X < LutLength; X++)
+			FFloat16Color* Row = &MipData[Y * InTextureWidth];
+			const float* Source = &InSourceData[Y * InTextureWidth * 3];
+			for (uint32 X = 0; X < InTextureWidth; X++)
 			{
-				Row[X] = FFloat16Color(FLinearColor(Source[X * 3 + 0], Source[X * 3 + 1], Source[X * 3 + 2]));
+				FFloat16Color& TargetColor = Row[X];
+				TargetColor.R.SetClamped(Source[X * 3 + 0]);
+				TargetColor.G.SetClamped(Source[X * 3 + 1]);
+				TargetColor.B.SetClamped(Source[X * 3 + 2]);
+				TargetColor.A.SetOne();
 			}
 		}
+		OutTexture->Source.UnlockMip(0);
 	}
-	Lut3dTexture->Source.UnlockMip(0);
 
 	//Generate a Guid from the identifier received from the library and our DDC version.
 	FGuid LutGuid;
-	GetOpenColorIOLUTKeyGuid(InLutIdentifier, LutGuid);
-	Lut3dTexture->Source.SetId(LutGuid, true);
+	GetOpenColorIOLUTKeyGuid(InProcessorIdentifier, InName, LutGuid);
+	OutTexture->Source.SetId(LutGuid, true);
 
 	//Process our new texture to be usable in rendering pipeline.
-	Lut3dTexture->UpdateResource();
+	OutTexture->UpdateResource();
 #endif
+
+	return OutTexture;
 }
 
 void UOpenColorIOColorTransform::FlushResourceShaderMaps()
diff --git a/Engine/Plugins/Compositing/OpenColorIO/Source/OpenColorIO/Private/OpenColorIOConfiguration.cpp b/Engine/Plugins/Compositing/OpenColorIO/Source/OpenColorIO/Private/OpenColorIOConfiguration.cpp
--- a/Engine/Plugins/Compositing/OpenColorIO/Source/OpenColorIO/Private/OpenColorIOConfiguration.cpp
+++ b/Engine/Plugins/Compositing/OpenColorIO/Source/OpenColorIO/Private/OpenColorIOConfiguration.cpp
@@ -54,7 +54,7 @@ void UOpenColorIOConfiguration::BeginDestroy()
 	Super::BeginDestroy();
 }
 
-bool UOpenColorIOConfiguration::GetShaderAndLUTResources(ERHIFeatureLevel::Type InFeatureLevel, const FString& InSourceColorSpace, const FString& InDestinationColorSpace, FOpenColorIOTransformResource*& OutShaderResource, FTextureResource*& OutLUT3dResource)
+bool UOpenColorIOConfiguration::GetRenderResources(ERHIFeatureLevel::Type InFeatureLevel, const FString& InSourceColorSpace, const FString& InDestinationColorSpace, FOpenColorIOTransformResource*& OutShaderResource, TArray<FTextureResource*>& OutTextureResources)
 {
 	UE_TRANSITIONAL_OBJECT_PTR(UOpenColorIOColorTransform)* TransformPtr = ColorTransforms.FindByPredicate([&](const UOpenColorIOColorTransform* InTransform)
 	{
@@ -68,7 +68,22 @@ bool UOpenColorIOConfiguration::GetShaderAndLUTResources(ERHIFeatureLevel::Type
 	}
 
 	UOpenColorIOColorTransform* Transform = *TransformPtr;
-	return Transform->GetShaderAndLUTResouces(InFeatureLevel, OutShaderResource, OutLUT3dResource);
+	return Transform->GetRenderResources(InFeatureLevel, OutShaderResource, OutTextureResources);
+}
+
+bool UOpenColorIOConfiguration::GetShaderAndLUTResources(ERHIFeatureLevel::Type InFeatureLevel, const FString& InSourceColorSpace, const FString& InDestinationColorSpace, FOpenColorIOTransformResource*& OutShaderResource, FTextureResource*& OutLUT3dResource)
+{
+	TArray<FTextureResource*> TextureResources;
+	if (GetRenderResources(InFeatureLevel, InSourceColorSpace, InDestinationColorSpace, OutShaderResource, TextureResources))
+	{
+		if (TextureResources.Num() > 0)
+		{
+			OutLUT3dResource = TextureResources[0];
+			return true;
+		}
+	}
+
+	return false;
 }
 
 bool UOpenColorIOConfiguration::HasTransform(const FString& InSourceColorSpace, const FString& InDestinationColorSpace)
diff --git a/Engine/Plugins/Compositing/OpenColorIO/Source/OpenColorIO/Private/OpenColorIODerivedDataVersion.h b/Engine/Plugins/Compositing/OpenColorIO/Source/OpenColorIO/Private/OpenColorIODerivedDataVersion.h
--- a/Engine/Plugins/Compositing/OpenColorIO/Source/OpenColorIO/Private/OpenColorIODerivedDataVersion.h
+++ b/Engine/Plugins/Compositing/OpenColorIO/Source/OpenColorIO/Private/OpenColorIODerivedDataVersion.h
@@ -10,4 +10,4 @@ OpenColorIOShaderDerivedDataVersion.h: Derived data version for OpenColorIO shad
 
 // In case of merge conflicts with DDC versions, you MUST generate a new GUID and set this new
 // guid as version
-#define OPENCOLORIO_DERIVEDDATA_VER		TEXT("BFFE3C29-C30F-4D6E-81D4-14BDE56F16ED")
\ No newline at end of file
+#define OPENCOLORIO_DERIVEDDATA_VER		TEXT("E164EF47-FAFA-4694-9F69-873A61E8F301")
\ No newline at end of file
diff --git a/Engine/Plugins/Compositing/OpenColorIO/Source/OpenColorIO/Private/OpenColorIODisplayExtension.cpp b/Engine/Plugins/Compositing/OpenColorIO/Source/OpenColorIO/Private/OpenColorIODisplayExtension.cpp
--- a/Engine/Plugins/Compositing/OpenColorIO/Source/OpenColorIO/Private/OpenColorIODisplayExtension.cpp
+++ b/Engine/Plugins/Compositing/OpenColorIO/Source/OpenColorIO/Private/OpenColorIODisplayExtension.cpp
@@ -55,7 +55,7 @@ void FOpenColorIODisplayExtension::SetupView(FSceneViewFamily& InViewFamily, FSc
 	//Cache render resource so they are available on the render thread (Can't access UObjects on RT)
 	//If something fails, cache invalid resources to invalidate them
 	FOpenColorIOTransformResource* ShaderResource = nullptr;
-	FTextureResource* LUT3dResource = nullptr;
+	TArray<FTextureResource*> TextureResources;
 
 	if (DisplayConfiguration.ColorConfiguration.ConfigurationSource == nullptr)
 	{
@@ -63,12 +63,12 @@ void FOpenColorIODisplayExtension::SetupView(FSceneViewFamily& InViewFamily, FSc
 	}
 	else
 	{
-		const bool bFoundTransform = DisplayConfiguration.ColorConfiguration.ConfigurationSource->GetShaderAndLUTResources(
+		const bool bFoundTransform = DisplayConfiguration.ColorConfiguration.ConfigurationSource->GetRenderResources(
 			InViewFamily.GetFeatureLevel()
 			, DisplayConfiguration.ColorConfiguration.SourceColorSpace.ColorSpaceName
 			, DisplayConfiguration.ColorConfiguration.DestinationColorSpace.ColorSpaceName
 			, ShaderResource
-			, LUT3dResource);
+			, TextureResources);
 
 		if (!bFoundTransform)
 		{
@@ -99,11 +99,11 @@ void FOpenColorIODisplayExtension::SetupView(FSceneViewFamily& InViewFamily, FSc
 	}
 
 	ENQUEUE_RENDER_COMMAND(ProcessColorSpaceTransform)(
-		[this, ShaderResource, LUT3dResource](FRHICommandListImmediate& RHICmdList)
+		[this, ShaderResource, TextureResources](FRHICommandListImmediate& RHICmdList)
 		{
 			//Caches render thread resource to be used when applying configuration in PostRenderViewFamily_RenderThread
 			CachedResourcesRenderThread.ShaderResource = ShaderResource;
-			CachedResourcesRenderThread.LUT3dResource = LUT3dResource;
+			CachedResourcesRenderThread.TextureResources = TextureResources;
 		}
 	);
 }
@@ -139,8 +139,6 @@ FScreenPassTexture FOpenColorIODisplayExtension::PostProcessPassAfterTonemap_Ren
 		return SceneColor;
 	}
 
-	RDG_EVENT_SCOPE(GraphBuilder, "OCIODisplayLook");
-
 	FScreenPassRenderTarget BackBufferRenderTarget;
 
 	// If the override output is provided it means that this is the last pass in post processing.
@@ -173,11 +171,7 @@ FScreenPassTexture FOpenColorIODisplayExtension::PostProcessPassAfterTonemap_Ren
 	FOpenColorIOPixelShaderParameters* Parameters = GraphBuilder.AllocParameters<FOpenColorIOPixelShaderParameters>();
 	Parameters->InputTexture = SceneColorRenderTarget.Texture;
 	Parameters->InputTextureSampler = TStaticSamplerState<>::GetRHI();
-	if (CachedResourcesRenderThread.LUT3dResource)
-	{
-		Parameters->Ocio_lut3d_0 = CachedResourcesRenderThread.LUT3dResource->TextureRHI;
-	}
-	Parameters->Ocio_lut3d_0Sampler = TStaticSamplerState<SF_Bilinear, AM_Clamp, AM_Clamp, AM_Clamp>::GetRHI();
+	OpenColorIOBindTextureResources(Parameters, CachedResourcesRenderThread.TextureResources);
 
 	// There is a special case where post processing and tonemapper are disabled. In this case tonemapper applies a static display Inverse of Gamma which defaults to 2.2.
 	// In the case when Both PostProcessing and ToneMapper are disabled we apply gamma manually. In every other case we apply inverse gamma before applying OCIO.
@@ -187,7 +181,7 @@ FScreenPassTexture FOpenColorIODisplayExtension::PostProcessPassAfterTonemap_Ren
 	checkSlow(View.bIsViewInfo);
 	const FViewInfo& ViewInfo = static_cast<const FViewInfo&>(View);
 
-	AddDrawScreenPass(GraphBuilder, RDG_EVENT_NAME("ProcessOCIOColorSpaceXfrm"), ViewInfo, BackBufferViewport, SceneColorViewport, OCIOPixelShader, Parameters);
+	AddDrawScreenPass(GraphBuilder, RDG_EVENT_NAME("OCIODisplayLook"), ViewInfo, BackBufferViewport, SceneColorViewport, OCIOPixelShader, Parameters);
 
 	return MoveTemp(BackBufferRenderTarget);
 }
\ No newline at end of file
diff --git a/Engine/Plugins/Compositing/OpenColorIO/Source/OpenColorIO/Private/OpenColorIORendering.cpp b/Engine/Plugins/Compositing/OpenColorIO/Source/OpenColorIO/Private/OpenColorIORendering.cpp
--- a/Engine/Plugins/Compositing/OpenColorIO/Source/OpenColorIO/Private/OpenColorIORendering.cpp
+++ b/Engine/Plugins/Compositing/OpenColorIO/Source/OpenColorIO/Private/OpenColorIORendering.cpp
@@ -13,6 +13,7 @@
 #include "OpenColorIOShader.h"
 #include "OpenColorIOShaderType.h"
 #include "OpenColorIOShared.h"
+#include "OpenColorIOColorTransform.h"
 #include "ScreenPass.h"
 #include "TextureResource.h"
 
@@ -21,7 +22,7 @@ void ProcessOCIOColorSpaceTransform_RenderThread(
 	FRHICommandListImmediate& InRHICmdList
 	, ERHIFeatureLevel::Type InFeatureLevel
 	, FOpenColorIOTransformResource* InOCIOColorTransformResource
-	, FTextureResource* InLUT3dResource
+	, TArray<FTextureResource*> InTextureResources
 	, FTextureRHIRef InputSpaceColorTexture
 	, FTextureRHIRef OutputSpaceColorTexture
 	, FIntPoint OutputResolution)
@@ -40,11 +41,7 @@ void ProcessOCIOColorSpaceTransform_RenderThread(
 	FOpenColorIOPixelShaderParameters* Parameters = GraphBuilder.AllocParameters<FOpenColorIOPixelShaderParameters>();
 	Parameters->InputTexture = InputTexture;
 	Parameters->InputTextureSampler = TStaticSamplerState<>::GetRHI();
-	if (InLUT3dResource != nullptr)
-	{
-		Parameters->Ocio_lut3d_0 = InLUT3dResource->TextureRHI;
-	}
-	Parameters->Ocio_lut3d_0Sampler = TStaticSamplerState<SF_Bilinear, AM_Clamp, AM_Clamp, AM_Clamp>::GetRHI();
+	OpenColorIOBindTextureResources(Parameters, InTextureResources);
 	// Set Gamma to 1., since we do not have any display parameters or requirement for Gamma.
 	Parameters->Gamma = 1.0;
 	Parameters->RenderTargets[0] = ScreenPassRenderTarget.GetRenderTargetBinding();
@@ -106,8 +103,8 @@ bool FOpenColorIORendering::ApplyColorTransform(UWorld* InWorld, const FOpenColo
 
 	const ERHIFeatureLevel::Type FeatureLevel = InWorld->Scene->GetFeatureLevel();
 	FOpenColorIOTransformResource* ShaderResource = nullptr;
-	FTextureResource* LUT3dResource = nullptr;
-	bool bFoundTransform = InSettings.ConfigurationSource->GetShaderAndLUTResources(FeatureLevel, InSettings.SourceColorSpace.ColorSpaceName, InSettings.DestinationColorSpace.ColorSpaceName, ShaderResource, LUT3dResource);
+	TArray<FTextureResource*> TextureResources;
+	bool bFoundTransform = InSettings.ConfigurationSource->GetRenderResources(FeatureLevel, InSettings.SourceColorSpace.ColorSpaceName, InSettings.DestinationColorSpace.ColorSpaceName, ShaderResource, TextureResources);
 	if (!bFoundTransform)
 	{
 		UE_LOG(LogOpenColorIO, Warning, TEXT("Can't apply color transform - Couldn't find shader to transform from %s to %s"), *InSettings.SourceColorSpace.ColorSpaceName, *InSettings.DestinationColorSpace.ColorSpaceName);
@@ -124,13 +121,13 @@ bool FOpenColorIORendering::ApplyColorTransform(UWorld* InWorld, const FOpenColo
 
 
 	ENQUEUE_RENDER_COMMAND(ProcessColorSpaceTransform)(
-		[FeatureLevel, InputResource, OutputResource, ShaderResource, LUT3dResource](FRHICommandListImmediate& RHICmdList)
+		[FeatureLevel, InputResource, OutputResource, ShaderResource, TextureResources](FRHICommandListImmediate& RHICmdList)
 		{
 			ProcessOCIOColorSpaceTransform_RenderThread(
 			RHICmdList,
 			FeatureLevel,
 			ShaderResource,
-			LUT3dResource,
+			TextureResources,
 			InputResource->TextureRHI,
 			OutputResource->TextureRHI,
 			FIntPoint(OutputResource->GetSizeX(), OutputResource->GetSizeY()));
diff --git a/Engine/Plugins/Compositing/OpenColorIO/Source/OpenColorIO/Private/OpenColorIOSettings.cpp b/Engine/Plugins/Compositing/OpenColorIO/Source/OpenColorIO/Private/OpenColorIOSettings.cpp
new file mode 100644
--- /dev/null
+++ b/Engine/Plugins/Compositing/OpenColorIO/Source/OpenColorIO/Private/OpenColorIOSettings.cpp
@@ -0,0 +1,26 @@
+// Copyright Epic Games, Inc. All Rights Reserved.
+
+#include "OpenColorIOSettings.h"
+
+UOpenColorIOSettings::UOpenColorIOSettings()
+	: bUseLegacyProcessor(false)
+{
+
+}
+
+FName UOpenColorIOSettings::GetCategoryName() const
+{
+	return TEXT("Plugins");
+}
+
+#if WITH_EDITOR
+FText UOpenColorIOSettings::GetSectionText() const
+{
+	return NSLOCTEXT("OpenColorIOSettings", "OpenColorIOSettingsSection", "OpenColorIO");
+}
+
+FName UOpenColorIOSettings::GetSectionName() const
+{
+	return TEXT("OpenColorIO");
+}
+#endif
diff --git a/Engine/Plugins/Compositing/OpenColorIO/Source/OpenColorIO/Private/OpenColorIOShader.cpp b/Engine/Plugins/Compositing/OpenColorIO/Source/OpenColorIO/Private/OpenColorIOShader.cpp
--- a/Engine/Plugins/Compositing/OpenColorIO/Source/OpenColorIO/Private/OpenColorIOShader.cpp
+++ b/Engine/Plugins/Compositing/OpenColorIO/Source/OpenColorIO/Private/OpenColorIOShader.cpp
@@ -2,4 +2,33 @@
 
 #include "OpenColorIOShader.h"
 
+#include "OpenColorIOModule.h"
+#include "OpenColorIOShared.h"
+#include "TextureResource.h"
+
+void OpenColorIOBindTextureResources(FOpenColorIOPixelShaderParameters* Parameters, TConstArrayView<FTextureResource*> InTextureResources)
+{
+	/**
+	* The engine's shader parameter API doesn't provide name/slot-based updates (in the style of legacy OpenGL),
+	* which is what OCIOv2 apparently relies on. Knowing that our input texture array follows the shader
+	* binding slot order, we set the FOpenColorIOPixelShaderParameters members by copying the relevant
+	* pointers using offsets from the initial "0" textures.
+	* 
+	* @todo: Improve binding logic, maybe using an approach similar to shaders generated by ComputeGraph.
+	*/
+	const uint8* BaseAddress = (uint8*)Parameters;
+
+	for (int32 Index = 0; Index < FMath::Min(InTextureResources.Num(), (int32)OpenColorIOShader::MaximumTextureSlots); ++Index)
+	{
+		const FTextureResource* Texture = InTextureResources[Index];
+		const bool bIsLUT3D = Texture->GetTexture3DResource() != nullptr;
+
+		const uint32 BaseOffset = bIsLUT3D ? offsetof(FOpenColorIOPixelShaderParameters, Ocio_lut3d_0) : offsetof(FOpenColorIOPixelShaderParameters, Ocio_lut1d_0);
+		const uint32 ParamOffset = BaseOffset + (sizeof(FRHITexture*) + sizeof(FRHISamplerState*)) * Index;
+
+		*((FRHITexture**)(BaseAddress + ParamOffset)) = Texture->TextureRHI.GetReference();
+		*((FRHISamplerState**)(BaseAddress + ParamOffset + sizeof(FRHITexture*))) = Texture->SamplerStateRHI.GetReference();
+	}
+}
+
 IMPLEMENT_SHADER_TYPE(, FOpenColorIOPixelShader, TEXT("/Plugin/OpenColorIO/Private/OpenColorIOShader.usf"), TEXT("MainPS"), SF_Pixel)
diff --git a/Engine/Plugins/Compositing/OpenColorIO/Source/OpenColorIO/Public/OpenColorIOColorTransform.h b/Engine/Plugins/Compositing/OpenColorIO/Source/OpenColorIO/Public/OpenColorIOColorTransform.h
--- a/Engine/Plugins/Compositing/OpenColorIO/Source/OpenColorIO/Public/OpenColorIOColorTransform.h
+++ b/Engine/Plugins/Compositing/OpenColorIO/Source/OpenColorIO/Public/OpenColorIOColorTransform.h
@@ -65,6 +65,9 @@ public:
 	/**
 	 * Returns the desired resources required to apply this transform during rendering.
 	 */
+	bool GetRenderResources(ERHIFeatureLevel::Type InFeatureLevel, FOpenColorIOTransformResource*& OutShaderResource, TArray<FTextureResource*>& OutTextureResources);
+
+	UE_DEPRECATED(5.1, "GetShaderAndLUTResouces is deprecated, please use GetRenderResources instead.")
 	bool GetShaderAndLUTResouces(ERHIFeatureLevel::Type InFeatureLevel, FOpenColorIOTransformResource*& OutShaderResource, FTextureResource*& OutLUT3dResource);
 
 	bool IsTransform(const FString& InSourceColorSpace, const FString& InDestinationColorSpace) const;
@@ -87,8 +90,11 @@ protected:
 	static void ProcessSerializedShaderMaps(UOpenColorIOColorTransform* Owner, TArray<FOpenColorIOTransformResource>& LoadedResources, FOpenColorIOTransformResource* (&OutColorTransformResourcesLoaded)[ERHIFeatureLevel::Num]);
 	
 	/**
-	 * Returns a Guid for the LUT based on its unique identifier and the OCIO DDC key.
+	 * Returns a Guid for the LUT based on its unique identifier, name and the OCIO DDC key.
 	 */
+	static void GetOpenColorIOLUTKeyGuid(const FString& InProcessorIdentifier, const FName& InName, FGuid& OutLutGuid );
+
+	UE_DEPRECATED(5.1, "This version of GetOpenColorIOLUTKeyGuid is deprecated. Please use the version that takes InName instead.")
 	static void GetOpenColorIOLUTKeyGuid(const FString& InLutIdentifier, FGuid& OutLutGuid);
 
 	/**
@@ -111,10 +117,18 @@ protected:
 	bool UpdateShaderInfo(FString& OutShaderCodeHash, FString& OutShaderCode, FString& OutRawConfigHash);
 
 	/**
-	 * Helper function taking raw LUT data coming from the library and initializing a UTexture with it.
+	 * Helper function taking raw 3D LUT data coming from the library and initializing a UVolumeTexture with it.
 	 */
+	TObjectPtr<UTexture> CreateTexture3DLUT(const FString& InProcessorIdentifier, const FName& InName, uint32 InLutLength, TextureFilter InFilter, const float* InSourceData);
+
+	UE_DEPRECATED(5.1, "Update3dLutTexture is deprecated, please use CreateTexture3DLUT instead.")
 	void Update3dLutTexture(const FString& InLutIdentifier, const float* InSourceData);
 
+	/**
+	 * Helper function taking raw 1D LUT data coming from the library and initializing a UTexture with it.
+	 */
+	TObjectPtr<UTexture> CreateTexture1DLUT(const FString& InProcessorIdentifier, const FName& InName, uint32 InTextureWidth, uint32 InTextureHeight, TextureFilter InFilter, bool bRedChannelOnly, const float* InSourceData );
+
 private:
 	void FlushResourceShaderMaps();
 
@@ -156,10 +170,10 @@ public:
 
 private:
 
-	/** If the color space requires a 3dLUT, this will contains the data to do the transform */
+	/** If the color space requires textures, this will contains the data to do the transform */
 	/** Note: This will be serialized when cooking. Otherwhise, it relies on raw data of the library and what's on DDC */
 	UPROPERTY(Transient)
-	TObjectPtr<UVolumeTexture> Lut3dTexture = nullptr;
+	TArray<TObjectPtr<UTexture>> Textures;
 	
 	/** Inline ColorTransform resources serialized from disk. To be processed on game thread in PostLoad. */
 	TArray<FOpenColorIOTransformResource> LoadedTransformResources;
diff --git a/Engine/Plugins/Compositing/OpenColorIO/Source/OpenColorIO/Public/OpenColorIOConfiguration.h b/Engine/Plugins/Compositing/OpenColorIO/Source/OpenColorIO/Public/OpenColorIOConfiguration.h
--- a/Engine/Plugins/Compositing/OpenColorIO/Source/OpenColorIO/Public/OpenColorIOConfiguration.h
+++ b/Engine/Plugins/Compositing/OpenColorIO/Source/OpenColorIO/Public/OpenColorIOConfiguration.h
@@ -35,6 +35,8 @@ public:
 	UOpenColorIOConfiguration(const FObjectInitializer& ObjectInitializer);
 
 public:
+	bool GetRenderResources(ERHIFeatureLevel::Type InFeatureLevel, const FString& InSourceColorSpace, const FString& InDestinationColorSpace, FOpenColorIOTransformResource*& OutShaderResource, TArray<FTextureResource*>& OutTextureResources);
+	UE_DEPRECATED(5.1, "GetShaderAndLUTResources is deprecated, please use GetRenderResources instead.")
 	bool GetShaderAndLUTResources(ERHIFeatureLevel::Type InFeatureLevel, const FString& InSourceColorSpace, const FString& InDestinationColorSpace, FOpenColorIOTransformResource*& OutShaderResource, FTextureResource*& OutLUT3dResource);
 	bool HasTransform(const FString& InSourceColorSpace, const FString& InDestinationColorSpace);
 	bool HasDesiredColorSpace(const FOpenColorIOColorSpace& ColorSpace) const;
diff --git a/Engine/Plugins/Compositing/OpenColorIO/Source/OpenColorIO/Public/OpenColorIODisplayExtension.h b/Engine/Plugins/Compositing/OpenColorIO/Source/OpenColorIO/Public/OpenColorIODisplayExtension.h
--- a/Engine/Plugins/Compositing/OpenColorIO/Source/OpenColorIO/Public/OpenColorIODisplayExtension.h
+++ b/Engine/Plugins/Compositing/OpenColorIO/Source/OpenColorIO/Public/OpenColorIODisplayExtension.h
@@ -65,7 +65,7 @@ private:
 	struct FCachedRenderResource
 	{
 		FOpenColorIOTransformResource* ShaderResource = nullptr;
-		FTextureResource* LUT3dResource = nullptr;
+		TArray<FTextureResource*> TextureResources;
 	};
 
 	/** Cached resources required to apply conversion for render thread */
diff --git a/Engine/Plugins/Compositing/OpenColorIO/Source/OpenColorIO/Public/OpenColorIOSettings.h b/Engine/Plugins/Compositing/OpenColorIO/Source/OpenColorIO/Public/OpenColorIOSettings.h
new file mode 100644
--- /dev/null
+++ b/Engine/Plugins/Compositing/OpenColorIO/Source/OpenColorIO/Public/OpenColorIOSettings.h
@@ -0,0 +1,38 @@
+// Copyright Epic Games, Inc. All Rights Reserved.
+
+#pragma once
+
+#include "CoreMinimal.h"
+#include "UObject/ObjectMacros.h"
+#include "Engine/DeveloperSettings.h"
+#include "Engine/EngineTypes.h"
+
+#include "OpenColorIOSettings.generated.h"
+
+/**
+ * Rendering settings.
+ */
+UCLASS(config = Engine, defaultconfig)
+class OPENCOLORIO_API UOpenColorIOSettings : public UDeveloperSettings
+{
+	GENERATED_BODY()
+
+public:
+	UOpenColorIOSettings();
+
+	//~ Begin UDeveloperSettings interface
+	virtual FName GetCategoryName() const;
+#if WITH_EDITOR
+	virtual FText GetSectionText() const override;
+	virtual FName GetSectionName() const override;
+#endif
+	//~ End UDeveloperSettings interface
+
+public:
+
+	UPROPERTY(config, EditAnywhere, Category = Transform, meta = (
+		DisplayName = "Enable Legacy Gpu Processor",
+		ToolTip = "Whether to enable OCIO V1's legacy gpu processor.",
+		ConfigRestartRequired = true))
+	uint8 bUseLegacyProcessor : 1;
+};
diff --git a/Engine/Plugins/Compositing/OpenColorIO/Source/OpenColorIO/Public/OpenColorIOShader.h b/Engine/Plugins/Compositing/OpenColorIO/Source/OpenColorIO/Public/OpenColorIOShader.h
--- a/Engine/Plugins/Compositing/OpenColorIO/Source/OpenColorIO/Public/OpenColorIOShader.h
+++ b/Engine/Plugins/Compositing/OpenColorIO/Source/OpenColorIO/Public/OpenColorIOShader.h
@@ -19,15 +19,51 @@ BEGIN_SHADER_PARAMETER_STRUCT(FOpenColorIOPixelShaderParameters, )
 	SHADER_PARAMETER_RDG_TEXTURE(Texture2D, InputTexture)
 	SHADER_PARAMETER_SAMPLER(SamplerState, InputTextureSampler)
 	SHADER_PARAMETER(float, Gamma)
+	
 	SHADER_PARAMETER_TEXTURE(Texture3D, Ocio_lut3d_0)
 	SHADER_PARAMETER_SAMPLER(SamplerState, Ocio_lut3d_0Sampler)
+	SHADER_PARAMETER_TEXTURE(Texture3D, Ocio_lut3d_1)
+	SHADER_PARAMETER_SAMPLER(SamplerState, Ocio_lut3d_1Sampler)
+	SHADER_PARAMETER_TEXTURE(Texture3D, Ocio_lut3d_2)
+	SHADER_PARAMETER_SAMPLER(SamplerState, Ocio_lut3d_2Sampler)
+	SHADER_PARAMETER_TEXTURE(Texture3D, Ocio_lut3d_3)
+	SHADER_PARAMETER_SAMPLER(SamplerState, Ocio_lut3d_3Sampler)
+	SHADER_PARAMETER_TEXTURE(Texture3D, Ocio_lut3d_4)
+	SHADER_PARAMETER_SAMPLER(SamplerState, Ocio_lut3d_4Sampler)
+	SHADER_PARAMETER_TEXTURE(Texture3D, Ocio_lut3d_5)
+	SHADER_PARAMETER_SAMPLER(SamplerState, Ocio_lut3d_5Sampler)
+	SHADER_PARAMETER_TEXTURE(Texture3D, Ocio_lut3d_6)
+	SHADER_PARAMETER_SAMPLER(SamplerState, Ocio_lut3d_6Sampler)
+	SHADER_PARAMETER_TEXTURE(Texture3D, Ocio_lut3d_7)
+	SHADER_PARAMETER_SAMPLER(SamplerState, Ocio_lut3d_7Sampler)
+
+	SHADER_PARAMETER_TEXTURE(Texture2D, Ocio_lut1d_0)
+	SHADER_PARAMETER_SAMPLER(SamplerState, Ocio_lut1d_0Sampler)
+	SHADER_PARAMETER_TEXTURE(Texture2D, Ocio_lut1d_1)
+	SHADER_PARAMETER_SAMPLER(SamplerState, Ocio_lut1d_1Sampler)
+	SHADER_PARAMETER_TEXTURE(Texture2D, Ocio_lut1d_2)
+	SHADER_PARAMETER_SAMPLER(SamplerState, Ocio_lut1d_2Sampler)
+	SHADER_PARAMETER_TEXTURE(Texture2D, Ocio_lut1d_3)
+	SHADER_PARAMETER_SAMPLER(SamplerState, Ocio_lut1d_3Sampler)
+	SHADER_PARAMETER_TEXTURE(Texture2D, Ocio_lut1d_4)
+	SHADER_PARAMETER_SAMPLER(SamplerState, Ocio_lut1d_4Sampler)
+	SHADER_PARAMETER_TEXTURE(Texture2D, Ocio_lut1d_5)
+	SHADER_PARAMETER_SAMPLER(SamplerState, Ocio_lut1d_5Sampler)
+	SHADER_PARAMETER_TEXTURE(Texture2D, Ocio_lut1d_6)
+	SHADER_PARAMETER_SAMPLER(SamplerState, Ocio_lut1d_6Sampler)
+	SHADER_PARAMETER_TEXTURE(Texture2D, Ocio_lut1d_7)
+	SHADER_PARAMETER_SAMPLER(SamplerState, Ocio_lut1d_7Sampler)
+	
 	RENDER_TARGET_BINDING_SLOTS()
 END_SHADER_PARAMETER_STRUCT()
 
 namespace OpenColorIOShader
 {
-	static const TCHAR* OpenColorIOShaderFunctionName = TEXT("OCIOConvert");
-	static const uint32 Lut3dEdgeLength = 65;
+	static constexpr TCHAR OpenColorIOShaderFunctionName[] = TEXT("OCIOConvert");
+	static constexpr uint32 Lut3dEdgeLength = 65;
+
+	// Combined maximum of either Ocio_lut1d or Ocio_lut3d, since the generated shaders will never use the same slots for both types.
+	static constexpr uint32 MaximumTextureSlots = 8;
 }
 
 class OPENCOLORIO_API FOpenColorIOPixelShader : public FGlobalShader
@@ -45,6 +81,4 @@ public:
 	}
 };
 
-
-
-
+OPENCOLORIO_API void OpenColorIOBindTextureResources(FOpenColorIOPixelShaderParameters* Parameters, TConstArrayView<FTextureResource*> InTextureResources);
