OCIO: Simple function to activate OCIO in game mode.

This removes the reliance on the VPUtilities plugin (which now come with non-trivial dependencies) for only utilizing OCIO in-game.

diff --git a/Engine/Plugins/Compositing/OpenColorIO/Source/OpenColorIO/Private/OpenColorIODisplayExtensionWrapper.cpp b/Engine/Plugins/Compositing/OpenColorIO/Source/OpenColorIO/Private/OpenColorIODisplayExtensionWrapper.cpp
--- a/Engine/Plugins/Compositing/OpenColorIO/Source/OpenColorIO/Private/OpenColorIODisplayExtensionWrapper.cpp
+++ b/Engine/Plugins/Compositing/OpenColorIO/Source/OpenColorIO/Private/OpenColorIODisplayExtensionWrapper.cpp
@@ -1,6 +1,12 @@
 // Copyright Epic Games, Inc. All Rights Reserved.
 
 #include "OpenColorIODisplayExtensionWrapper.h"
+#include "Engine/GameEngine.h"
+#include "Slate/SceneViewport.h"
+
+#if WITH_EDITOR
+#include "Editor.h"
+#endif // WITH_EDITOR
 
 void UOpenColorIODisplayExtensionWrapper::CreateDisplayExtensionIfNotExists()
 {
@@ -61,4 +67,41 @@ UOpenColorIODisplayExtensionWrapper* UOpenColorIODisplayExtensionWrapper::Create
 	OutExtension->SetSceneExtensionIsActiveFunction(IsActiveFunction);
 
 	return OutExtension;
-}
\ No newline at end of file
+}
+
+UOpenColorIODisplayExtensionWrapper* UOpenColorIODisplayExtensionWrapper::CreateInGameOpenColorIODisplayExtension(FOpenColorIODisplayConfiguration InDisplayConfiguration)
+{
+	FSceneViewExtensionIsActiveFunctor IsActiveFunctor;
+	IsActiveFunctor.IsActiveFunction = [](const ISceneViewExtension* SceneViewExtension, const FSceneViewExtensionContext& Context)
+	{
+		if (!Context.Viewport)
+		{
+			return TOptional<bool>();
+		}
+
+#if WITH_EDITOR
+		if (GIsEditor && GEditor)
+		{
+			// Activate the SVE if it is a PIE viewport.
+			if (Context.Viewport->IsPlayInEditorViewport())
+			{
+				return TOptional<bool>(true);
+			}
+		}
+#endif // WITH_EDITOR
+
+		// Activate the SVE if it is the game's primary viewport.
+		if (UGameEngine* GameEngine = Cast<UGameEngine>(GEngine))
+		{
+			if (GameEngine->SceneViewport.IsValid() && (GameEngine->SceneViewport->GetViewport() == Context.Viewport))
+			{
+				return TOptional<bool>(true);
+			}
+		}
+
+		// If our viewport did not meet any of the criteria to activate the SVE, emit no opinion.
+		return TOptional<bool>();
+	};
+
+	return CreateOpenColorIODisplayExtension(InDisplayConfiguration, IsActiveFunctor);
+}
diff --git a/Engine/Plugins/Compositing/OpenColorIO/Source/OpenColorIO/Public/OpenColorIODisplayExtensionWrapper.h b/Engine/Plugins/Compositing/OpenColorIO/Source/OpenColorIO/Public/OpenColorIODisplayExtensionWrapper.h
--- a/Engine/Plugins/Compositing/OpenColorIO/Source/OpenColorIO/Public/OpenColorIODisplayExtensionWrapper.h
+++ b/Engine/Plugins/Compositing/OpenColorIO/Source/OpenColorIO/Public/OpenColorIODisplayExtensionWrapper.h
@@ -40,6 +40,10 @@ public:
 		FOpenColorIODisplayConfiguration InDisplayConfiguration, 
 		const FSceneViewExtensionIsActiveFunctor& IsActiveFunction);
 
+	// Creates an instance of this object, configured for use in game with the given OCIO configuration.
+	UFUNCTION(BlueprintCallable, meta = (DisplayName = "Create In-Game OpenColorIO Display Extension"), Category = "OpenColorIO")
+	static UOpenColorIODisplayExtensionWrapper* CreateInGameOpenColorIODisplayExtension(
+		FOpenColorIODisplayConfiguration InDisplayConfiguration);
 private:
 
 	// Creates OCIO Extension if it doesn't exist. 
