OCIO: Avoid render target clear unless needed in the scene view extension pass.

Minor cleanup of the screen pass. We now also export "FViewInfo::GetOverwriteLoadAction" so that it can be used by such SVE passes.

#rnx

diff --git a/Engine/Plugins/Compositing/OpenColorIO/Source/OpenColorIO/Private/OpenColorIODisplayExtension.cpp b/Engine/Plugins/Compositing/OpenColorIO/Source/OpenColorIO/Private/OpenColorIODisplayExtension.cpp
--- a/Engine/Plugins/Compositing/OpenColorIO/Source/OpenColorIO/Private/OpenColorIODisplayExtension.cpp
+++ b/Engine/Plugins/Compositing/OpenColorIO/Source/OpenColorIO/Private/OpenColorIODisplayExtension.cpp
@@ -118,10 +118,9 @@ void FOpenColorIODisplayExtension::SubscribeToPostProcessingPass(EPostProcessing
 
 FScreenPassTexture FOpenColorIODisplayExtension::PostProcessPassAfterTonemap_RenderThread(FRDGBuilder& GraphBuilder, const FSceneView& View, const FPostProcessMaterialInputs& InOutInputs)
 {
-
 	const FSceneViewFamily& ViewFamily = *View.Family;
-
-	const FScreenPassTexture& SceneColor = InOutInputs.Textures[(uint32)EPostProcessMaterialInput::SceneColor];
+	const FScreenPassTexture& SceneColor = InOutInputs.GetInput(EPostProcessMaterialInput::SceneColor);
+	check(SceneColor.IsValid());
 
 	//If the shader resource could not be found, skip this frame. The LUT isn't required
 	if (CachedResourcesRenderThread.ShaderResource == nullptr)
@@ -134,54 +133,35 @@ FScreenPassTexture FOpenColorIODisplayExtension::PostProcessPassAfterTonemap_Ren
 		return SceneColor;
 	}
 
-	if (!SceneColor.IsValid())
-	{
-		return SceneColor;
-	}
-
-	FScreenPassRenderTarget BackBufferRenderTarget;
+	checkSlow(View.bIsViewInfo);
+	const FViewInfo& ViewInfo = static_cast<const FViewInfo&>(View);
+	
+	FScreenPassRenderTarget Output = InOutInputs.OverrideOutput;
 
 	// If the override output is provided it means that this is the last pass in post processing.
-	if (InOutInputs.OverrideOutput.IsValid())
+	if (!Output.IsValid())
 	{
-		BackBufferRenderTarget = InOutInputs.OverrideOutput;
+		Output = FScreenPassRenderTarget::CreateFromInput(GraphBuilder, SceneColor, ViewInfo.GetOverwriteLoadAction(), TEXT("OCIORenderTarget"));
 	}
-	else
-	{
-		// Reusing the same output description for our back buffer as SceneColor when it's not overriden
-		FRDGTextureDesc OutputDesc = SceneColor.Texture->Desc;
-		OutputDesc.Flags |= TexCreate_RenderTargetable;
-		FLinearColor ClearColor(0., 0., 0., 0.);
-		OutputDesc.ClearValue = FClearValueBinding(ClearColor);
-
-		FRDGTexture* BackBufferRenderTargetTexture = GraphBuilder.CreateTexture(OutputDesc, TEXT("BackBufferRenderTargetTexture"));
-		BackBufferRenderTarget = FScreenPassRenderTarget(BackBufferRenderTargetTexture, SceneColor.ViewRect, ERenderTargetLoadAction::EClear);
-	}
-
-	//Get input and output viewports. Backbuffer could be targeting a different region than input viewport
-	const FScreenPassTextureViewport SceneColorViewport(SceneColor);
-	const FScreenPassTextureViewport BackBufferViewport(BackBufferRenderTarget);
 
-	FScreenPassRenderTarget SceneColorRenderTarget(SceneColor, ERenderTargetLoadAction::ELoad);
+	const FScreenPassTextureViewport InputViewport(SceneColor);
+	const FScreenPassTextureViewport OutputViewport(Output);
 
 	TShaderRef<FOpenColorIOPixelShader> OCIOPixelShader = CachedResourcesRenderThread.ShaderResource->GetShader<FOpenColorIOPixelShader>();
 
 	const float DisplayGamma = View.Family->RenderTarget->GetDisplayGamma();
 
 	FOpenColorIOPixelShaderParameters* Parameters = GraphBuilder.AllocParameters<FOpenColorIOPixelShaderParameters>();
-	Parameters->InputTexture = SceneColorRenderTarget.Texture;
+	Parameters->InputTexture = SceneColor.Texture;
 	Parameters->InputTextureSampler = TStaticSamplerState<>::GetRHI();
 	OpenColorIOBindTextureResources(Parameters, CachedResourcesRenderThread.TextureResources);
 
 	// There is a special case where post processing and tonemapper are disabled. In this case tonemapper applies a static display Inverse of Gamma which defaults to 2.2.
 	// In the case when Both PostProcessing and ToneMapper are disabled we apply gamma manually. In every other case we apply inverse gamma before applying OCIO.
 	Parameters->Gamma = (ViewFamily.EngineShowFlags.Tonemapper == 0) || (ViewFamily.EngineShowFlags.PostProcessing == 0) ? DefaultDisplayGamma : DefaultDisplayGamma / DisplayGamma;
-	Parameters->RenderTargets[0] = BackBufferRenderTarget.GetRenderTargetBinding();
-
-	checkSlow(View.bIsViewInfo);
-	const FViewInfo& ViewInfo = static_cast<const FViewInfo&>(View);
+	Parameters->RenderTargets[0] = Output.GetRenderTargetBinding();
 
-	AddDrawScreenPass(GraphBuilder, RDG_EVENT_NAME("OCIODisplayLook"), ViewInfo, BackBufferViewport, SceneColorViewport, OCIOPixelShader, Parameters);
+	AddDrawScreenPass(GraphBuilder, RDG_EVENT_NAME("OCIODisplayLook"), ViewInfo, OutputViewport, InputViewport, OCIOPixelShader, Parameters);
 
-	return MoveTemp(BackBufferRenderTarget);
+	return MoveTemp(Output);
 }
\ No newline at end of file
diff --git a/Engine/Source/Runtime/Renderer/Private/SceneRendering.h b/Engine/Source/Runtime/Renderer/Private/SceneRendering.h
--- a/Engine/Source/Runtime/Renderer/Private/SceneRendering.h
+++ b/Engine/Source/Runtime/Renderer/Private/SceneRendering.h
@@ -1641,7 +1641,7 @@ public:
 	void EnqueueEyeAdaptationExposureBufferReadback(FRDGBuilder& GraphBuilder) const;
 	
 	/** Returns the load action to use when overwriting all pixels of a target that you intend to read from. Takes into account the HMD hidden area mesh. */
-	ERenderTargetLoadAction GetOverwriteLoadAction() const;
+	RENDERER_API ERenderTargetLoadAction GetOverwriteLoadAction() const;
 
 	/** Informs sceneinfo that tonemapping LUT has queued commands to compute it at least once */
 	void SetValidTonemappingLUT() const;
