OCIO: Store the LUT textures in a slot-index map instead of an array.

Fixes the assumption that generated shaders would always start at index 0. Some cases where the starting index is 1 were found while adding support for display-view transforms.

#rnx

diff --git a/Engine/Plugins/Compositing/OpenColorIO/Source/OpenColorIO/Private/OpenColorIOColorTransform.cpp b/Engine/Plugins/Compositing/OpenColorIO/Source/OpenColorIO/Private/OpenColorIOColorTransform.cpp
--- a/Engine/Plugins/Compositing/OpenColorIO/Source/OpenColorIO/Private/OpenColorIOColorTransform.cpp
+++ b/Engine/Plugins/Compositing/OpenColorIO/Source/OpenColorIO/Private/OpenColorIOColorTransform.cpp
@@ -2,6 +2,7 @@
 
 #include "OpenColorIOColorTransform.h"
 
+#include "Containers/SortedMap.h"
 #include "Materials/MaterialInterface.h"
 #include "Math/PackedVector.h"
 #include "Modules/ModuleManager.h"
@@ -277,9 +278,11 @@ void UOpenColorIOColorTransform::CacheResourceTextures()
 						{
 							Filter = TF_Nearest;
 						}
-						TObjectPtr<UTexture> Result = CreateTexture3DLUT(ProcessorID, FName(TextureName), EdgeLength, Filter, TextureValues);
+						const FName TextureFName = FName(TextureName);
+						TObjectPtr<UTexture> Result = CreateTexture3DLUT(ProcessorID, TextureFName, EdgeLength, Filter, TextureValues);
 
-						Textures.Add(MoveTemp(Result));
+						const int32 SlotIndex = TextureFName.GetNumber() - 1; // Rely on FName's index number extraction for convenience
+						Textures.Add(SlotIndex, MoveTemp(Result));
 					}
 
 					// Process 1D luts
@@ -301,20 +304,14 @@ void UOpenColorIOColorTransform::CacheResourceTextures()
 
 						TextureFilter Filter = Interpolation == OCIO_NAMESPACE::Interpolation::INTERP_NEAREST ? TF_Nearest : TF_Bilinear;
 						bool bRedChannelOnly = Channel == OCIO_NAMESPACE::GpuShaderCreator::TEXTURE_RED_CHANNEL;
-						TObjectPtr<UTexture> Result = CreateTexture1DLUT(ProcessorID, FName(TextureName), TextureWidth, TextureHeight, Filter, bRedChannelOnly, TextureValues);
+						const FName TextureFName = FName(TextureName);
+						TObjectPtr<UTexture> Result = CreateTexture1DLUT(ProcessorID, TextureFName, TextureWidth, TextureHeight, Filter, bRedChannelOnly, TextureValues);
 
-						Textures.Add(MoveTemp(Result));
+						const int32 SlotIndex = TextureFName.GetNumber() - 1; // Rely on FName's index number extraction for convenience
+						Textures.Add(SlotIndex, MoveTemp(Result));
 					}
 
 					ensureAlwaysMsgf(Textures.Num() <= (int32)OpenColorIOShader::MaximumTextureSlots, TEXT("Color transform %s exceeds our current limit of %u texture slots. Use the legacy processor instead."), *GetTransformFriendlyName(), OpenColorIOShader::MaximumTextureSlots);
-
-					// Subscript numbers extracted by FName can be used to sort textures for matching the shader binding accesses.
-					Textures.Sort(
-						[](const UTexture& TextureA, const UTexture& TextureB)
-						{
-							return TextureA.GetFName().GetNumber() < TextureB.GetFName().GetNumber();
-						}
-					);
 				}
 				else
 				{
@@ -405,16 +402,16 @@ FOpenColorIOTransformResource* UOpenColorIOColorTransform::AllocateResource()
 	return new FOpenColorIOTransformResource();
 }
 
-bool UOpenColorIOColorTransform::GetRenderResources(ERHIFeatureLevel::Type InFeatureLevel, FOpenColorIOTransformResource*& OutShaderResource, TArray<FTextureResource*>& OutTextureResources)
+bool UOpenColorIOColorTransform::GetRenderResources(ERHIFeatureLevel::Type InFeatureLevel, FOpenColorIOTransformResource*& OutShaderResource, TSortedMap<int32, FTextureResource*>& OutTextureResources)
 {
 	OutShaderResource = ColorTransformResources[InFeatureLevel];
 	if (OutShaderResource)
 	{
 		OutTextureResources.Reserve(Textures.Num());
 
-		for (const TObjectPtr<UTexture>& Texture : Textures)
+		for (const TPair<int32, TObjectPtr<UTexture>>& Pair : Textures)
 		{
-			OutTextureResources.Add(Texture->GetResource());
+			OutTextureResources.Add(Pair.Key, Pair.Value->GetResource());
 		}
 
 		return true;
@@ -428,10 +425,10 @@ bool UOpenColorIOColorTransform::GetRenderResources(ERHIFeatureLevel::Type InFea
 
 bool UOpenColorIOColorTransform::GetShaderAndLUTResouces(ERHIFeatureLevel::Type InFeatureLevel, FOpenColorIOTransformResource*& OutShaderResource, FTextureResource*& OutLUT3dResource)
 {
-	TArray<FTextureResource*> OutTextureResources;
+	TSortedMap<int32, FTextureResource*> OutTextureResources;
 	bool bResult = GetRenderResources(InFeatureLevel, OutShaderResource, OutTextureResources);
 
-	if (!OutTextureResources.IsEmpty())
+	if (OutTextureResources.Contains(0))
 	{
 		OutLUT3dResource = OutTextureResources[0];
 	}
@@ -612,7 +609,7 @@ void UOpenColorIOColorTransform::Update3dLutTexture(const FString& InLutIdentifi
 	Textures.Empty();
 
 	TObjectPtr<UTexture> Texture = CreateTexture3DLUT(InLutIdentifier, FName(), OpenColorIOShader::Lut3dEdgeLength, TextureFilter::TF_Bilinear, InSourceData);
-	Textures.Add(MoveTemp(Texture));
+	Textures.Add(0, MoveTemp(Texture));
 #endif
 }
 
diff --git a/Engine/Plugins/Compositing/OpenColorIO/Source/OpenColorIO/Private/OpenColorIOConfiguration.cpp b/Engine/Plugins/Compositing/OpenColorIO/Source/OpenColorIO/Private/OpenColorIOConfiguration.cpp
--- a/Engine/Plugins/Compositing/OpenColorIO/Source/OpenColorIO/Private/OpenColorIOConfiguration.cpp
+++ b/Engine/Plugins/Compositing/OpenColorIO/Source/OpenColorIO/Private/OpenColorIOConfiguration.cpp
@@ -2,6 +2,7 @@
 
 #include "OpenColorIOConfiguration.h"
 
+#include "Containers/SortedMap.h"
 #include "EngineAnalytics.h"
 #include "Engine/VolumeTexture.h"
 #include "Framework/Notifications/NotificationManager.h"
@@ -54,16 +55,16 @@ void UOpenColorIOConfiguration::BeginDestroy()
 	Super::BeginDestroy();
 }
 
-bool UOpenColorIOConfiguration::GetRenderResources(ERHIFeatureLevel::Type InFeatureLevel, const FString& InSourceColorSpace, const FString& InDestinationColorSpace, FOpenColorIOTransformResource*& OutShaderResource, TArray<FTextureResource*>& OutTextureResources)
+bool UOpenColorIOConfiguration::GetRenderResources(ERHIFeatureLevel::Type InFeatureLevel, const FOpenColorIOColorConversionSettings& InSettings, FOpenColorIOTransformResource*& OutShaderResource, TSortedMap<int32, FTextureResource*>& OutTextureResources)
 {
 	TObjectPtr<UOpenColorIOColorTransform>* TransformPtr = ColorTransforms.FindByPredicate([&](const UOpenColorIOColorTransform* InTransform)
 	{
-		return InTransform->SourceColorSpace == InSourceColorSpace && InTransform->DestinationColorSpace == InDestinationColorSpace;
+		return InTransform->SourceColorSpace == InSettings.SourceColorSpace.ColorSpaceName && InTransform->DestinationColorSpace == InSettings.DestinationColorSpace.ColorSpaceName;
 	});
 
 	if (TransformPtr == nullptr)
 	{
-		UE_LOG(LogOpenColorIO, Warning, TEXT("Color transform data from %s to %s was not found."), *InSourceColorSpace, *InDestinationColorSpace);
+		UE_LOG(LogOpenColorIO, Warning, TEXT("Color transform data from %s to %s was not found."), *InSettings.SourceColorSpace.ColorSpaceName, *InSettings.DestinationColorSpace.ColorSpaceName);
 		return false;
 	}
 
@@ -73,10 +74,13 @@ bool UOpenColorIOConfiguration::GetRenderResources(ERHIFeatureLevel::Type InFeat
 
 bool UOpenColorIOConfiguration::GetShaderAndLUTResources(ERHIFeatureLevel::Type InFeatureLevel, const FString& InSourceColorSpace, const FString& InDestinationColorSpace, FOpenColorIOTransformResource*& OutShaderResource, FTextureResource*& OutLUT3dResource)
 {
-	TArray<FTextureResource*> TextureResources;
-	if (GetRenderResources(InFeatureLevel, InSourceColorSpace, InDestinationColorSpace, OutShaderResource, TextureResources))
+	TSortedMap<int32, FTextureResource*> TextureResources;
+	FOpenColorIOColorConversionSettings Settings;
+	Settings.SourceColorSpace.ColorSpaceName = InSourceColorSpace;
+	Settings.DestinationColorSpace.ColorSpaceName = InDestinationColorSpace;
+	if (GetRenderResources(InFeatureLevel, Settings, OutShaderResource, TextureResources))
 	{
-		if (TextureResources.Num() > 0)
+		if (TextureResources.Contains(0))
 		{
 			OutLUT3dResource = TextureResources[0];
 			return true;
diff --git a/Engine/Plugins/Compositing/OpenColorIO/Source/OpenColorIO/Private/OpenColorIODerivedDataVersion.h b/Engine/Plugins/Compositing/OpenColorIO/Source/OpenColorIO/Private/OpenColorIODerivedDataVersion.h
--- a/Engine/Plugins/Compositing/OpenColorIO/Source/OpenColorIO/Private/OpenColorIODerivedDataVersion.h
+++ b/Engine/Plugins/Compositing/OpenColorIO/Source/OpenColorIO/Private/OpenColorIODerivedDataVersion.h
@@ -10,4 +10,4 @@ OpenColorIOShaderDerivedDataVersion.h: Derived data version for OpenColorIO shad
 
 // In case of merge conflicts with DDC versions, you MUST generate a new GUID and set this new
 // guid as version
-#define OPENCOLORIO_DERIVEDDATA_VER		TEXT("E164EF47-FAFA-4694-9F69-873A61E8F301")
\ No newline at end of file
+#define OPENCOLORIO_DERIVEDDATA_VER		TEXT("2018266A-4ABF-4AA9-B1B4-695931DAAFF7")
\ No newline at end of file
diff --git a/Engine/Plugins/Compositing/OpenColorIO/Source/OpenColorIO/Private/OpenColorIODisplayExtension.cpp b/Engine/Plugins/Compositing/OpenColorIO/Source/OpenColorIO/Private/OpenColorIODisplayExtension.cpp
--- a/Engine/Plugins/Compositing/OpenColorIO/Source/OpenColorIO/Private/OpenColorIODisplayExtension.cpp
+++ b/Engine/Plugins/Compositing/OpenColorIO/Source/OpenColorIO/Private/OpenColorIODisplayExtension.cpp
@@ -2,6 +2,7 @@
 
 #include "OpenColorIODisplayExtension.h"
 
+#include "Containers/SortedMap.h"
 #include "CoreGlobals.h"
 #include "OpenColorIOConfiguration.h"
 #include "OpenColorIOColorTransform.h"
@@ -55,7 +56,7 @@ void FOpenColorIODisplayExtension::SetupView(FSceneViewFamily& InViewFamily, FSc
 	//Cache render resource so they are available on the render thread (Can't access UObjects on RT)
 	//If something fails, cache invalid resources to invalidate them
 	FOpenColorIOTransformResource* ShaderResource = nullptr;
-	TArray<FTextureResource*> TextureResources;
+	TSortedMap<int32, FTextureResource*> TransformTextureResources;
 
 	if (DisplayConfiguration.ColorConfiguration.ConfigurationSource == nullptr)
 	{
@@ -65,10 +66,9 @@ void FOpenColorIODisplayExtension::SetupView(FSceneViewFamily& InViewFamily, FSc
 	{
 		const bool bFoundTransform = DisplayConfiguration.ColorConfiguration.ConfigurationSource->GetRenderResources(
 			InViewFamily.GetFeatureLevel()
-			, DisplayConfiguration.ColorConfiguration.SourceColorSpace.ColorSpaceName
-			, DisplayConfiguration.ColorConfiguration.DestinationColorSpace.ColorSpaceName
+			, DisplayConfiguration.ColorConfiguration
 			, ShaderResource
-			, TextureResources);
+			, TransformTextureResources);
 
 		if (!bFoundTransform)
 		{
@@ -99,7 +99,7 @@ void FOpenColorIODisplayExtension::SetupView(FSceneViewFamily& InViewFamily, FSc
 	}
 
 	ENQUEUE_RENDER_COMMAND(ProcessColorSpaceTransform)(
-		[this, ShaderResource, TextureResources](FRHICommandListImmediate& RHICmdList)
+		[this, ShaderResource, TextureResources = MoveTemp(TransformTextureResources)](FRHICommandListImmediate& RHICmdList)
 		{
 			//Caches render thread resource to be used when applying configuration in PostRenderViewFamily_RenderThread
 			CachedResourcesRenderThread.ShaderResource = ShaderResource;
diff --git a/Engine/Plugins/Compositing/OpenColorIO/Source/OpenColorIO/Private/OpenColorIORendering.cpp b/Engine/Plugins/Compositing/OpenColorIO/Source/OpenColorIO/Private/OpenColorIORendering.cpp
--- a/Engine/Plugins/Compositing/OpenColorIO/Source/OpenColorIO/Private/OpenColorIORendering.cpp
+++ b/Engine/Plugins/Compositing/OpenColorIO/Source/OpenColorIO/Private/OpenColorIORendering.cpp
@@ -22,7 +22,7 @@ void ProcessOCIOColorSpaceTransform_RenderThread(
 	FRHICommandListImmediate& InRHICmdList
 	, ERHIFeatureLevel::Type InFeatureLevel
 	, FOpenColorIOTransformResource* InOCIOColorTransformResource
-	, TArray<FTextureResource*> InTextureResources
+	, const TSortedMap<int32, FTextureResource*>& InTextureResources
 	, FTextureRHIRef InputSpaceColorTexture
 	, FTextureRHIRef OutputSpaceColorTexture
 	, FIntPoint OutputResolution)
@@ -103,8 +103,8 @@ bool FOpenColorIORendering::ApplyColorTransform(UWorld* InWorld, const FOpenColo
 
 	const ERHIFeatureLevel::Type FeatureLevel = InWorld->Scene->GetFeatureLevel();
 	FOpenColorIOTransformResource* ShaderResource = nullptr;
-	TArray<FTextureResource*> TextureResources;
-	bool bFoundTransform = InSettings.ConfigurationSource->GetRenderResources(FeatureLevel, InSettings.SourceColorSpace.ColorSpaceName, InSettings.DestinationColorSpace.ColorSpaceName, ShaderResource, TextureResources);
+	TSortedMap<int32, FTextureResource*> TextureResources;
+	bool bFoundTransform = InSettings.ConfigurationSource->GetRenderResources(FeatureLevel, InSettings, ShaderResource, TextureResources);
 	if (!bFoundTransform)
 	{
 		UE_LOG(LogOpenColorIO, Warning, TEXT("Can't apply color transform - Couldn't find shader to transform from %s to %s"), *InSettings.SourceColorSpace.ColorSpaceName, *InSettings.DestinationColorSpace.ColorSpaceName);
diff --git a/Engine/Plugins/Compositing/OpenColorIO/Source/OpenColorIO/Private/OpenColorIOShader.cpp b/Engine/Plugins/Compositing/OpenColorIO/Source/OpenColorIO/Private/OpenColorIOShader.cpp
--- a/Engine/Plugins/Compositing/OpenColorIO/Source/OpenColorIO/Private/OpenColorIOShader.cpp
+++ b/Engine/Plugins/Compositing/OpenColorIO/Source/OpenColorIO/Private/OpenColorIOShader.cpp
@@ -2,11 +2,12 @@
 
 #include "OpenColorIOShader.h"
 
+#include "Containers/SortedMap.h"
 #include "OpenColorIOModule.h"
 #include "OpenColorIOShared.h"
 #include "TextureResource.h"
 
-void OpenColorIOBindTextureResources(FOpenColorIOPixelShaderParameters* Parameters, TConstArrayView<FTextureResource*> InTextureResources)
+void OpenColorIOBindTextureResources(FOpenColorIOPixelShaderParameters* Parameters, const TSortedMap<int32, FTextureResource*>& InTextureResources)
 {
 	/**
 	* The engine's shader parameter API doesn't provide name/slot-based updates (in the style of legacy OpenGL),
@@ -18,9 +19,12 @@ void OpenColorIOBindTextureResources(FOpenColorIOPixelShaderParameters* Paramete
 	*/
 	const uint8* BaseAddress = (uint8*)Parameters;
 
-	for (int32 Index = 0; Index < FMath::Min(InTextureResources.Num(), (int32)OpenColorIOShader::MaximumTextureSlots); ++Index)
+	for ( const TPair<int32, FTextureResource*>& Pair : InTextureResources )
 	{
-		const FTextureResource* Texture = InTextureResources[Index];
+		const int32 Index = Pair.Key;
+		ensureMsgf(Index >= 0 && Index < (int32)OpenColorIOShader::MaximumTextureSlots, TEXT("Out of range LUT texture slot index."));
+
+		const FTextureResource* Texture = Pair.Value;
 		const bool bIsLUT3D = Texture->GetTexture3DResource() != nullptr;
 
 		const uint32 BaseOffset = bIsLUT3D ? offsetof(FOpenColorIOPixelShaderParameters, Ocio_lut3d_0) : offsetof(FOpenColorIOPixelShaderParameters, Ocio_lut1d_0);
diff --git a/Engine/Plugins/Compositing/OpenColorIO/Source/OpenColorIO/Public/OpenColorIOColorTransform.h b/Engine/Plugins/Compositing/OpenColorIO/Source/OpenColorIO/Public/OpenColorIOColorTransform.h
--- a/Engine/Plugins/Compositing/OpenColorIO/Source/OpenColorIO/Public/OpenColorIOColorTransform.h
+++ b/Engine/Plugins/Compositing/OpenColorIO/Source/OpenColorIO/Public/OpenColorIOColorTransform.h
@@ -65,7 +65,7 @@ public:
 	/**
 	 * Returns the desired resources required to apply this transform during rendering.
 	 */
-	bool GetRenderResources(ERHIFeatureLevel::Type InFeatureLevel, FOpenColorIOTransformResource*& OutShaderResource, TArray<FTextureResource*>& OutTextureResources);
+	bool GetRenderResources(ERHIFeatureLevel::Type InFeatureLevel, FOpenColorIOTransformResource*& OutShaderResource, TSortedMap<int32, FTextureResource*>& OutTextureResources);
 
 	UE_DEPRECATED(5.1, "GetShaderAndLUTResouces is deprecated, please use GetRenderResources instead.")
 	bool GetShaderAndLUTResouces(ERHIFeatureLevel::Type InFeatureLevel, FOpenColorIOTransformResource*& OutShaderResource, FTextureResource*& OutLUT3dResource);
@@ -173,7 +173,7 @@ private:
 	/** If the color space requires textures, this will contains the data to do the transform */
 	/** Note: This will be serialized when cooking. Otherwhise, it relies on raw data of the library and what's on DDC */
 	UPROPERTY(Transient)
-	TArray<TObjectPtr<UTexture>> Textures;
+	TMap<int32, TObjectPtr<UTexture>> Textures;
 	
 	/** Inline ColorTransform resources serialized from disk. To be processed on game thread in PostLoad. */
 	TArray<FOpenColorIOTransformResource> LoadedTransformResources;
diff --git a/Engine/Plugins/Compositing/OpenColorIO/Source/OpenColorIO/Public/OpenColorIOConfiguration.h b/Engine/Plugins/Compositing/OpenColorIO/Source/OpenColorIO/Public/OpenColorIOConfiguration.h
--- a/Engine/Plugins/Compositing/OpenColorIO/Source/OpenColorIO/Public/OpenColorIOConfiguration.h
+++ b/Engine/Plugins/Compositing/OpenColorIO/Source/OpenColorIO/Public/OpenColorIOConfiguration.h
@@ -35,7 +35,7 @@ public:
 	UOpenColorIOConfiguration(const FObjectInitializer& ObjectInitializer);
 
 public:
-	bool GetRenderResources(ERHIFeatureLevel::Type InFeatureLevel, const FString& InSourceColorSpace, const FString& InDestinationColorSpace, FOpenColorIOTransformResource*& OutShaderResource, TArray<FTextureResource*>& OutTextureResources);
+	bool GetRenderResources(ERHIFeatureLevel::Type InFeatureLevel, const FOpenColorIOColorConversionSettings& InSettings, FOpenColorIOTransformResource*& OutShaderResource, TSortedMap<int32, FTextureResource*>& OutTextureResources);
 	UE_DEPRECATED(5.1, "GetShaderAndLUTResources is deprecated, please use GetRenderResources instead.")
 	bool GetShaderAndLUTResources(ERHIFeatureLevel::Type InFeatureLevel, const FString& InSourceColorSpace, const FString& InDestinationColorSpace, FOpenColorIOTransformResource*& OutShaderResource, FTextureResource*& OutLUT3dResource);
 	bool HasTransform(const FString& InSourceColorSpace, const FString& InDestinationColorSpace);
diff --git a/Engine/Plugins/Compositing/OpenColorIO/Source/OpenColorIO/Public/OpenColorIODisplayExtension.h b/Engine/Plugins/Compositing/OpenColorIO/Source/OpenColorIO/Public/OpenColorIODisplayExtension.h
--- a/Engine/Plugins/Compositing/OpenColorIO/Source/OpenColorIO/Public/OpenColorIODisplayExtension.h
+++ b/Engine/Plugins/Compositing/OpenColorIO/Source/OpenColorIO/Public/OpenColorIODisplayExtension.h
@@ -65,7 +65,7 @@ private:
 	struct FCachedRenderResource
 	{
 		FOpenColorIOTransformResource* ShaderResource = nullptr;
-		TArray<FTextureResource*> TextureResources;
+		TSortedMap<int32, FTextureResource*> TextureResources;
 	};
 
 	/** Cached resources required to apply conversion for render thread */
diff --git a/Engine/Plugins/Compositing/OpenColorIO/Source/OpenColorIO/Public/OpenColorIOShader.h b/Engine/Plugins/Compositing/OpenColorIO/Source/OpenColorIO/Public/OpenColorIOShader.h
--- a/Engine/Plugins/Compositing/OpenColorIO/Source/OpenColorIO/Public/OpenColorIOShader.h
+++ b/Engine/Plugins/Compositing/OpenColorIO/Source/OpenColorIO/Public/OpenColorIOShader.h
@@ -6,6 +6,7 @@
 
 #pragma once
 
+#include "Containers/ContainersFwd.h"
 #include "GlobalShader.h"
 #include "OpenColorIOShaderType.h"
 #include "OpenColorIOShared.h"
@@ -81,4 +82,4 @@ public:
 	}
 };
 
-OPENCOLORIO_API void OpenColorIOBindTextureResources(FOpenColorIOPixelShaderParameters* Parameters, TConstArrayView<FTextureResource*> InTextureResources);
+OPENCOLORIO_API void OpenColorIOBindTextureResources(FOpenColorIOPixelShaderParameters* Parameters, const TSortedMap<int32, FTextureResource*>& InTextureResources);
