OCIO: Prevent users from creating an invalid display extension, with warning.

#rnx

diff --git a/Engine/Plugins/Compositing/OpenColorIO/Source/OpenColorIO/Private/OpenColorIODisplayExtensionWrapper.cpp b/Engine/Plugins/Compositing/OpenColorIO/Source/OpenColorIO/Private/OpenColorIODisplayExtensionWrapper.cpp
--- a/Engine/Plugins/Compositing/OpenColorIO/Source/OpenColorIO/Private/OpenColorIODisplayExtensionWrapper.cpp
+++ b/Engine/Plugins/Compositing/OpenColorIO/Source/OpenColorIO/Private/OpenColorIODisplayExtensionWrapper.cpp
@@ -1,6 +1,7 @@
 // Copyright Epic Games, Inc. All Rights Reserved.
 
 #include "OpenColorIODisplayExtensionWrapper.h"
+#include "OpenColorIOModule.h"
 #include "Engine/GameEngine.h"
 #include "Slate/SceneViewport.h"
 
@@ -58,6 +59,12 @@ UOpenColorIODisplayExtensionWrapper* UOpenColorIODisplayExtensionWrapper::Create
 	FOpenColorIODisplayConfiguration InDisplayConfiguration,
 	const FSceneViewExtensionIsActiveFunctor& IsActiveFunction)
 {
+	if (!InDisplayConfiguration.ColorConfiguration.IsValid())
+	{
+		UE_LOG(LogOpenColorIO, Warning, TEXT("%s, no display extension was created."), *InDisplayConfiguration.ColorConfiguration.ToString());
+		return nullptr;
+	}
+
 	// Create OCIO Scene View Extension and configure it.
 
 	UOpenColorIODisplayExtensionWrapper* OutExtension = NewObject<UOpenColorIODisplayExtensionWrapper>();
