OCIO: support 32f LUTs without conversions to 16f, which is now available in 5.1.

Option to use 32f resources is exposed in plugin settings.

diff --git a/Engine/Plugins/Compositing/OpenColorIO/Source/OpenColorIO/Private/OpenColorIOColorTransform.cpp b/Engine/Plugins/Compositing/OpenColorIO/Source/OpenColorIO/Private/OpenColorIOColorTransform.cpp
--- a/Engine/Plugins/Compositing/OpenColorIO/Source/OpenColorIO/Private/OpenColorIOColorTransform.cpp
+++ b/Engine/Plugins/Compositing/OpenColorIO/Source/OpenColorIO/Private/OpenColorIOColorTransform.cpp
@@ -556,6 +556,7 @@ TObjectPtr<UTexture> UOpenColorIOColorTransform::CreateTexture3DLUT(const FStrin
 
 #if WITH_EDITOR && WITH_OCIO
 	check(InSourceData);
+	const UOpenColorIOSettings* Settings = GetDefault<UOpenColorIOSettings>();
 
 	/* Note here that while it is possible to create proper 32f textures using [UTexture]::CreateTransient and reparenting
 	 * via Texture->Rename(nullptr, this), cooking would fail as it remains unsupported currently for those formats.
@@ -567,25 +568,34 @@ TObjectPtr<UTexture> UOpenColorIOColorTransform::CreateTexture3DLUT(const FStrin
 	//If not, the source data will be used to generate the platform data.
 	OutTexture->MipGenSettings = TMGS_NoMipmaps;
 	OutTexture->SRGB = 0;
-	OutTexture->CompressionNone = true;
+	OutTexture->LODGroup = TEXTUREGROUP_ColorLookupTable;
+	if (Settings->bUse32fLUT)
+	{
+		// 32f resources have to be explicitely requested using this compression setting.
+		OutTexture->CompressionSettings = TextureCompressionSettings::TC_HDR_F32;
+	}
+	else
+	{
+		OutTexture->CompressionNone = true;
+	}
 	OutTexture->Filter = InFilter;
 	OutTexture->AddressMode = TextureAddress::TA_Clamp;
-	OutTexture->Source.Init(InLutLength, InLutLength, InLutLength, /*NumMips=*/ 1, TSF_RGBA16F, nullptr);
+	OutTexture->Source.Init(InLutLength, InLutLength, InLutLength, /*NumMips=*/ 1, TSF_RGBA32F, nullptr);
 
-	FFloat16Color* MipData = reinterpret_cast<FFloat16Color*>(OutTexture->Source.LockMip(0));
+	FLinearColor* MipData = reinterpret_cast<FLinearColor*>(OutTexture->Source.LockMip(0));
 	for (uint32 Z = 0; Z < InLutLength; ++Z)
 	{
 		for (uint32 Y = 0; Y < InLutLength; Y++)
 		{
-			FFloat16Color* Row = &MipData[Y * InLutLength + Z * InLutLength * InLutLength];
+			FLinearColor* Row = &MipData[Y * InLutLength + Z * InLutLength * InLutLength];
 			const float* Source = &InSourceData[Y * InLutLength * 3 + Z * InLutLength * InLutLength * 3];
 			for (uint32 X = 0; X < InLutLength; X++)
 			{
-				FFloat16Color& TargetColor = Row[X];
-				TargetColor.R.SetClamped(Source[X * 3 + 0]);
-				TargetColor.G.SetClamped(Source[X * 3 + 1]);
-				TargetColor.B.SetClamped(Source[X * 3 + 2]);
-				TargetColor.A.SetOne();
+				FLinearColor& TargetColor = Row[X];
+				TargetColor.R = Source[X * 3 + 0];
+				TargetColor.G = Source[X * 3 + 1];
+				TargetColor.B = Source[X * 3 + 2];
+				TargetColor.A = 1.0f;
 			}
 		}
 	}
@@ -619,11 +629,21 @@ TObjectPtr<UTexture> UOpenColorIOColorTransform::CreateTexture1DLUT(const FStrin
 
 #if WITH_EDITOR && WITH_OCIO
 	check(InSourceData);
+	const UOpenColorIOSettings* Settings = GetDefault<UOpenColorIOSettings>();
 
 	OutTexture = NewObject<UTexture2D>(this, InName);
 	OutTexture->MipGenSettings = TMGS_NoMipmaps;
 	OutTexture->SRGB = 0;
-	OutTexture->CompressionNone = true;
+	OutTexture->LODGroup = TEXTUREGROUP_ColorLookupTable;
+	if (Settings->bUse32fLUT)
+	{
+		// 32f resources have to be explicitely requested using this compression setting.
+		OutTexture->CompressionSettings = bRedChannelOnly ? TextureCompressionSettings::TC_SingleFloat : TextureCompressionSettings::TC_HDR_F32;
+	}
+	else
+	{
+		OutTexture->CompressionNone = true;
+	}
 	OutTexture->Filter = InFilter;
 	OutTexture->AddressX = TextureAddress::TA_Clamp;
 	OutTexture->AddressY = TextureAddress::TA_Clamp;
@@ -634,20 +654,20 @@ TObjectPtr<UTexture> UOpenColorIOColorTransform::CreateTexture1DLUT(const FStrin
 	}
 	else
 	{
-		OutTexture->Source.Init(InTextureWidth, InTextureHeight, /*NumSlices=*/ 1, /*NumMips=*/ 1, TSF_RGBA16F, nullptr);
+		OutTexture->Source.Init(InTextureWidth, InTextureHeight, /*NumSlices=*/ 1, /*NumMips=*/ 1, TSF_RGBA32F, nullptr);
 
-		FFloat16Color* MipData = reinterpret_cast<FFloat16Color*>(OutTexture->Source.LockMip(0));
+		FLinearColor* MipData = reinterpret_cast<FLinearColor*>(OutTexture->Source.LockMip(0));
 		for (uint32 Y = 0; Y < InTextureHeight; Y++)
 		{
-			FFloat16Color* Row = &MipData[Y * InTextureWidth];
+			FLinearColor* Row = &MipData[Y * InTextureWidth];
 			const float* Source = &InSourceData[Y * InTextureWidth * 3];
 			for (uint32 X = 0; X < InTextureWidth; X++)
 			{
-				FFloat16Color& TargetColor = Row[X];
-				TargetColor.R.SetClamped(Source[X * 3 + 0]);
-				TargetColor.G.SetClamped(Source[X * 3 + 1]);
-				TargetColor.B.SetClamped(Source[X * 3 + 2]);
-				TargetColor.A.SetOne();
+				FLinearColor& TargetColor = Row[X];
+				TargetColor.R = Source[X * 3 + 0];
+				TargetColor.G = Source[X * 3 + 1];
+				TargetColor.B = Source[X * 3 + 2];
+				TargetColor.A = 1.0f;
 			}
 		}
 		OutTexture->Source.UnlockMip(0);
diff --git a/Engine/Plugins/Compositing/OpenColorIO/Source/OpenColorIO/Private/OpenColorIOSettings.cpp b/Engine/Plugins/Compositing/OpenColorIO/Source/OpenColorIO/Private/OpenColorIOSettings.cpp
--- a/Engine/Plugins/Compositing/OpenColorIO/Source/OpenColorIO/Private/OpenColorIOSettings.cpp
+++ b/Engine/Plugins/Compositing/OpenColorIO/Source/OpenColorIO/Private/OpenColorIOSettings.cpp
@@ -4,6 +4,7 @@
 
 UOpenColorIOSettings::UOpenColorIOSettings()
 	: bUseLegacyProcessor(false)
+	, bUse32fLUT(false)
 {
 
 }
diff --git a/Engine/Plugins/Compositing/OpenColorIO/Source/OpenColorIO/Public/OpenColorIOSettings.h b/Engine/Plugins/Compositing/OpenColorIO/Source/OpenColorIO/Public/OpenColorIOSettings.h
--- a/Engine/Plugins/Compositing/OpenColorIO/Source/OpenColorIO/Public/OpenColorIOSettings.h
+++ b/Engine/Plugins/Compositing/OpenColorIO/Source/OpenColorIO/Public/OpenColorIOSettings.h
@@ -35,4 +35,10 @@ public:
 		ToolTip = "Whether to enable OCIO V1's legacy gpu processor.",
 		ConfigRestartRequired = true))
 	uint8 bUseLegacyProcessor : 1;
+
+	UPROPERTY(config, EditAnywhere, Category = Transform, meta = (
+		DisplayName = "32-bit float LUTs",
+		ToolTip = "Whether to create lookup table texture resources in 32-bit float format (higher performance requirements).",
+		ConfigRestartRequired = true))
+	uint8 bUse32fLUT : 1;
 };
