OCIO: Adding support for display and views.

At the simplest level, displays and views are simply pairs of strings that are used as a different method of creating a gpu processor (see GetTransformProcessor in OpenColorIOTransform.cpp). This shelf extends support for this accross configuration and property menus, while remaining backwards compatible with existing use of configs.

(Depends on patched library in CL 21992958.)

diff --git a/Engine/Plugins/Compositing/OpenColorIO/Source/OpenColorIO/Private/OpenColorIOColorSpace.cpp b/Engine/Plugins/Compositing/OpenColorIO/Source/OpenColorIO/Private/OpenColorIOColorSpace.cpp
--- a/Engine/Plugins/Compositing/OpenColorIO/Source/OpenColorIO/Private/OpenColorIOColorSpace.cpp
+++ b/Engine/Plugins/Compositing/OpenColorIO/Source/OpenColorIO/Private/OpenColorIOColorSpace.cpp
@@ -24,7 +24,7 @@ FString FOpenColorIOColorSpace::ToString() const
 {
 	if (IsValid())
 	{
-		return FString::Printf(TEXT("%s"), *ColorSpaceName);
+		return ColorSpaceName;
 	}
 	return TEXT("<Invalid>");
 }
@@ -63,6 +63,43 @@ FString FOpenColorIOColorSpace::GetFamilyNameAtDepth(int32 InDepth) const
 	return ReturnName;
 }
 
+/*
+ * FOpenColorIODisplayView implementation
+ */
+
+FOpenColorIODisplayView::FOpenColorIODisplayView()
+	: Display()
+	, View()
+{
+}
+
+FOpenColorIODisplayView::FOpenColorIODisplayView(FStringView InDisplayName, FStringView InViewName)
+	: Display(InDisplayName)
+	, View(InViewName)
+{
+}
+
+FString FOpenColorIODisplayView::ToString() const
+{
+	if (IsValid())
+	{
+		return Display + TEXT(" - ") + View;
+	}
+
+	return TEXT("<Invalid>");
+}
+
+bool FOpenColorIODisplayView::IsValid() const
+{
+	return !Display.IsEmpty() && !View.IsEmpty();
+}
+
+void FOpenColorIODisplayView::Reset()
+{
+	Display.Reset();
+	View.Reset();
+}
+
 /*
  * FOpenColorIOColorConversionSettings implementation
  */
@@ -77,7 +114,14 @@ FString FOpenColorIOColorConversionSettings::ToString() const
 {
 	if (ConfigurationSource)
 	{
-		return FString::Printf(TEXT("%s config - %s to %s"), *ConfigurationSource->GetName(), *SourceColorSpace.ToString(), *DestinationColorSpace.ToString());
+		if (IsDisplayView())
+		{
+			return FString::Printf(TEXT("%s config - %s to %s"), *ConfigurationSource->GetName(), *SourceColorSpace.ToString(), *DestinationDisplayView.ToString());
+		}
+		else
+		{
+			return FString::Printf(TEXT("%s config - %s to %s"), *ConfigurationSource->GetName(), *SourceColorSpace.ToString(), *DestinationColorSpace.ToString());
+		}
 	}
 	return TEXT("<Invalid Conversion>");
 }
@@ -86,7 +130,14 @@ bool FOpenColorIOColorConversionSettings::IsValid() const
 {
 	if (ConfigurationSource)
 	{
-		return ConfigurationSource->HasTransform(SourceColorSpace.ColorSpaceName, DestinationColorSpace.ColorSpaceName);
+		if (IsDisplayView())
+		{
+			return ConfigurationSource->HasTransform(SourceColorSpace.ColorSpaceName, DestinationDisplayView.Display, DestinationDisplayView.View);
+		}
+		else
+		{
+			return ConfigurationSource->HasTransform(SourceColorSpace.ColorSpaceName, DestinationColorSpace.ColorSpaceName);
+		}
 	}
 
 	return false;
@@ -101,14 +152,32 @@ void FOpenColorIOColorConversionSettings::ValidateColorSpaces()
 			SourceColorSpace.Reset();
 		}
 
-		if (!ConfigurationSource->HasDesiredColorSpace(DestinationColorSpace))
+		if (IsDisplayView())
 		{
+			if (!ConfigurationSource->HasDesiredDisplayView(DestinationDisplayView))
+			{
+				DestinationDisplayView.Reset();
+			}
 			DestinationColorSpace.Reset();
 		}
+		else
+		{
+			if (!ConfigurationSource->HasDesiredColorSpace(DestinationColorSpace))
+			{
+				DestinationColorSpace.Reset();
+			}
+			DestinationDisplayView.Reset();
+		}
 	}
 	else
 	{
 		SourceColorSpace.Reset();
 		DestinationColorSpace.Reset();
+		DestinationDisplayView.Reset();
 	}
 }
+
+bool FOpenColorIOColorConversionSettings::IsDisplayView() const
+{
+	return DestinationDisplayView.IsValid();
+}
diff --git a/Engine/Plugins/Compositing/OpenColorIO/Source/OpenColorIO/Private/OpenColorIOColorTransform.cpp b/Engine/Plugins/Compositing/OpenColorIO/Source/OpenColorIO/Private/OpenColorIOColorTransform.cpp
--- a/Engine/Plugins/Compositing/OpenColorIO/Source/OpenColorIO/Private/OpenColorIOColorTransform.cpp
+++ b/Engine/Plugins/Compositing/OpenColorIO/Source/OpenColorIO/Private/OpenColorIOColorTransform.cpp
@@ -38,6 +38,21 @@ namespace {
 
 		return static_cast<OptimizationFlags>(OptimizationFlags::OPTIMIZATION_DEFAULT | OptimizationFlags::OPTIMIZATION_NO_DYNAMIC_PROPERTIES);
 	}
+
+	OCIO_NAMESPACE::ConstProcessorRcPtr GetTransformProcessor(UOpenColorIOColorTransform* InTransform, const OCIO_NAMESPACE::ConstConfigRcPtr& CurrentConfig)
+	{
+		if (InTransform->bIsDisplayViewType)
+		{
+			return CurrentConfig->getProcessor(StringCast<ANSICHAR>(*InTransform->SourceColorSpace).Get(),
+				StringCast<ANSICHAR>(*InTransform->Display).Get(),
+				StringCast<ANSICHAR>(*InTransform->View).Get(),
+				OCIO_NAMESPACE::TransformDirection::TRANSFORM_DIR_FORWARD);
+		}
+		else
+		{
+			return CurrentConfig->getProcessor(StringCast<ANSICHAR>(*InTransform->SourceColorSpace).Get(), StringCast<ANSICHAR>(*InTransform->DestinationColorSpace).Get());
+		}
+	}
 #endif
 }
 
@@ -150,6 +165,13 @@ bool UOpenColorIOColorTransform::Initialize(UOpenColorIOConfiguration* InOwner,
 	return GenerateColorTransformData(InSourceColorSpace, InDestinationColorSpace);
 }
 
+bool UOpenColorIOColorTransform::Initialize(UOpenColorIOConfiguration* InOwner, const FString& InSourceColorSpace, const FString& InDisplay, const FString& InView, bool bInForwardDirection)
+{
+	check(InOwner);
+	ConfigurationOwner = InOwner;
+	return GenerateColorTransformData(InSourceColorSpace, InDisplay, InView, bInForwardDirection);
+}
+
 void UOpenColorIOColorTransform::Serialize(FArchive& Ar)
 {
 	Super::Serialize(Ar);
@@ -226,8 +248,7 @@ void UOpenColorIOColorTransform::CacheResourceTextures()
 			try
 #endif
 			{
-				OCIO_NAMESPACE::ConstProcessorRcPtr TransformProcessor = CurrentConfig->getProcessor(StringCast<ANSICHAR>(*SourceColorSpace).Get(), StringCast<ANSICHAR>(*DestinationColorSpace).Get());
-				
+				OCIO_NAMESPACE::ConstProcessorRcPtr TransformProcessor = GetTransformProcessor(this, CurrentConfig);				
 				if (TransformProcessor)
 				{
 					OCIO_NAMESPACE::GpuShaderDescRcPtr ShaderDescription = OCIO_NAMESPACE::GpuShaderDesc::CreateShaderDesc();
@@ -437,7 +458,22 @@ bool UOpenColorIOColorTransform::GetShaderAndLUTResouces(ERHIFeatureLevel::Type
 
 bool UOpenColorIOColorTransform::IsTransform(const FString& InSourceColorSpace, const FString& InDestinationColorSpace) const
 {
-	return SourceColorSpace == InSourceColorSpace && DestinationColorSpace == InDestinationColorSpace;
+	if (!bIsDisplayViewType)
+	{
+		return SourceColorSpace == InSourceColorSpace && DestinationColorSpace == InDestinationColorSpace;
+	}
+
+	return false;
+}
+
+bool UOpenColorIOColorTransform::IsTransform(const FString& InSourceColorSpace, const FString& InDisplay, const FString& InView, bool bInForwardDirection) const
+{
+	if (bIsDisplayViewType)
+	{
+		return SourceColorSpace == InSourceColorSpace && Display == InDisplay && View == InView;
+	}
+
+	return false;
 }
 
 void UOpenColorIOColorTransform::AllColorTransformsCacheResourceShadersForRendering()
@@ -460,6 +496,28 @@ bool UOpenColorIOColorTransform::GenerateColorTransformData(const FString& InSou
 
 	SourceColorSpace = InSourceColorSpace;
 	DestinationColorSpace = InDestinationColorSpace;
+	bIsDisplayViewType = false;
+
+	CacheResourceTextures();
+	CacheResourceShadersForRendering(true);
+
+	return true;
+#endif //WITH_EDITOR
+	return false;
+}
+
+bool UOpenColorIOColorTransform::GenerateColorTransformData(const FString& InSourceColorSpace, const FString& InDisplay, const FString& InView, bool bInForwardDirection)
+{
+#if WITH_EDITOR && WITH_OCIO
+	if (InSourceColorSpace.IsEmpty() || InDisplay.IsEmpty() || InView.IsEmpty())
+	{
+		return false;
+	}
+
+	SourceColorSpace = InSourceColorSpace;
+	Display = InDisplay;
+	View = InView;
+	bIsDisplayViewType = true;
 
 	CacheResourceTextures();
 	CacheResourceShadersForRendering(true);
@@ -471,7 +529,14 @@ bool UOpenColorIOColorTransform::GenerateColorTransformData(const FString& InSou
 
 FString UOpenColorIOColorTransform::GetTransformFriendlyName()
 {
-	return SourceColorSpace + TEXT(" to ") + DestinationColorSpace;
+	if (bIsDisplayViewType)
+	{
+		return SourceColorSpace + TEXT(" to ") + Display + TEXT(" - ") + View;
+	}
+	else
+	{
+		return SourceColorSpace + TEXT(" to ") + DestinationColorSpace;
+	}
 }
 
 bool UOpenColorIOColorTransform::UpdateShaderInfo(FString& OutShaderCodeHash, FString& OutShaderCode, FString& OutRawConfigHash)
@@ -485,7 +550,7 @@ bool UOpenColorIOColorTransform::UpdateShaderInfo(FString& OutShaderCodeHash, FS
 		try
 #endif
 		{
-			OCIO_NAMESPACE::ConstProcessorRcPtr TransformProcessor = CurrentConfig->getProcessor(StringCast<ANSICHAR>(*SourceColorSpace).Get(), StringCast<ANSICHAR>(*DestinationColorSpace).Get());
+			OCIO_NAMESPACE::ConstProcessorRcPtr TransformProcessor = GetTransformProcessor(this, CurrentConfig);
 			if (TransformProcessor)
 			{
 				OCIO_NAMESPACE::GpuShaderDescRcPtr ShaderDescription = OCIO_NAMESPACE::GpuShaderDesc::CreateShaderDesc();
diff --git a/Engine/Plugins/Compositing/OpenColorIO/Source/OpenColorIO/Private/OpenColorIOConfiguration.cpp b/Engine/Plugins/Compositing/OpenColorIO/Source/OpenColorIO/Private/OpenColorIOConfiguration.cpp
--- a/Engine/Plugins/Compositing/OpenColorIO/Source/OpenColorIO/Private/OpenColorIOConfiguration.cpp
+++ b/Engine/Plugins/Compositing/OpenColorIO/Source/OpenColorIO/Private/OpenColorIOConfiguration.cpp
@@ -59,12 +59,19 @@ bool UOpenColorIOConfiguration::GetRenderResources(ERHIFeatureLevel::Type InFeat
 {
 	TObjectPtr<UOpenColorIOColorTransform>* TransformPtr = ColorTransforms.FindByPredicate([&](const UOpenColorIOColorTransform* InTransform)
 	{
-		return InTransform->SourceColorSpace == InSettings.SourceColorSpace.ColorSpaceName && InTransform->DestinationColorSpace == InSettings.DestinationColorSpace.ColorSpaceName;
+		if(InTransform->bIsDisplayViewType)
+		{
+			return InTransform->SourceColorSpace == InSettings.SourceColorSpace.ColorSpaceName && InTransform->Display == InSettings.DestinationDisplayView.Display && InTransform->View == InSettings.DestinationDisplayView.View;
+		}
+		else
+		{
+			return InTransform->SourceColorSpace == InSettings.SourceColorSpace.ColorSpaceName && InTransform->DestinationColorSpace == InSettings.DestinationColorSpace.ColorSpaceName;
+		}
 	});
 
 	if (TransformPtr == nullptr)
 	{
-		UE_LOG(LogOpenColorIO, Warning, TEXT("Color transform data from %s to %s was not found."), *InSettings.SourceColorSpace.ColorSpaceName, *InSettings.DestinationColorSpace.ColorSpaceName);
+		UE_LOG(LogOpenColorIO, Warning, TEXT("Color transform for [%s] was not found."), *InSettings.SourceColorSpace.ToString());
 		return false;
 	}
 
@@ -100,11 +107,26 @@ bool UOpenColorIOConfiguration::HasTransform(const FString& InSourceColorSpace,
 	return (TransformData != nullptr);
 }
 
+bool UOpenColorIOConfiguration::HasTransform(const FString& InSourceColorSpace, const FString& InDisplay, const FString& InView)
+{
+	UE_TRANSITIONAL_OBJECT_PTR(UOpenColorIOColorTransform)* TransformData = ColorTransforms.FindByPredicate([&](const UOpenColorIOColorTransform* InTransformData)
+		{
+			return InTransformData->IsTransform(InSourceColorSpace, InDisplay, InView);
+		});
+
+	return (TransformData != nullptr);
+}
+
 bool UOpenColorIOConfiguration::HasDesiredColorSpace(const FOpenColorIOColorSpace& ColorSpace) const
 {
 	return DesiredColorSpaces.Find(ColorSpace) != INDEX_NONE;
 }
 
+bool UOpenColorIOConfiguration::HasDesiredDisplayView(const FOpenColorIODisplayView& DisplayView) const
+{
+	return DesiredDisplayViews.Find(DisplayView) != INDEX_NONE;
+}
+
 bool UOpenColorIOConfiguration::Validate() const
 {
 #if WITH_EDITOR 
@@ -130,7 +152,9 @@ void UOpenColorIOConfiguration::ReloadExistingColorspaces()
 {
 #if WITH_EDITOR && WITH_OCIO
 	TArray<FOpenColorIOColorSpace> ColorSpacesToBeReloaded = DesiredColorSpaces;
+	TArray<FOpenColorIODisplayView> DisplayViewsToBeReloaded = DesiredDisplayViews;
 	DesiredColorSpaces.Reset();
+	DesiredDisplayViews.Reset();
 	CleanupTransforms();
 	LoadConfigurationFile();
 
@@ -157,6 +181,20 @@ void UOpenColorIOConfiguration::ReloadExistingColorspaces()
 		DesiredColorSpaces.Add(ColorSpace);
 	}
 
+	for (const FOpenColorIODisplayView& ExistingDisplayView : DisplayViewsToBeReloaded)
+	{
+		const auto DisplayName = StringCast<ANSICHAR>(*ExistingDisplayView.Display);
+		const auto ViewName = StringCast<ANSICHAR>(*ExistingDisplayView.View);
+		const auto TransformName = LoadedConfig->getDisplayViewTransformName(DisplayName.Get(), ViewName.Get());
+		if (TransformName == nullptr)
+		{
+			// Name not found, therefore we don't need to re-add this display-view.
+			continue;
+		}
+
+		DesiredDisplayViews.Add(ExistingDisplayView);
+	}
+
 	// Genereate new shaders.
 	for (int32 indexTop = 0; indexTop < DesiredColorSpaces.Num(); ++indexTop)
 	{
@@ -169,6 +207,12 @@ void UOpenColorIOConfiguration::ReloadExistingColorspaces()
 			CreateColorTransform(TopColorSpace.ColorSpaceName, OtherColorSpace.ColorSpaceName);
 			CreateColorTransform(OtherColorSpace.ColorSpaceName, TopColorSpace.ColorSpaceName);
 		}
+
+		for (const FOpenColorIODisplayView& DisplayView : DesiredDisplayViews)
+		{
+			// note: we only support display-view transforms in the forward direction currently.
+			CreateColorTransform(TopColorSpace.ColorSpaceName, DisplayView.Display, DisplayView.View);
+		}
 	}
 #endif
 }
@@ -286,38 +330,74 @@ void UOpenColorIOConfiguration::CreateColorTransform(const FString& InSourceColo
 	}
 }
 
-void UOpenColorIOConfiguration::CleanupTransforms()
+void UOpenColorIOConfiguration::CreateColorTransform(const FString& InSourceColorSpace, const FString& InDisplay, const FString& InView, bool bForwardDirection)
 {
-	for (int32 i = 0; i < ColorTransforms.Num(); ++i)
+	if (InSourceColorSpace.IsEmpty() || InDisplay.IsEmpty() || InView.IsEmpty())
 	{
-		UOpenColorIOColorTransform* Transform = ColorTransforms[i];
+		return;
+	}
 
-		//Check if the source color space of this transform is desired. If not, remove that transform from the list and move on.
-		const FOpenColorIOColorSpace* FoundSourceColorPtr = DesiredColorSpaces.FindByPredicate([&](const FOpenColorIOColorSpace& InOtherColorSpace)
-		{
-			return InOtherColorSpace.ColorSpaceName == Transform->SourceColorSpace;
-		});
+	if (HasTransform(InSourceColorSpace, InDisplay, InView))
+	{
+		UE_LOG(LogOpenColorIO, Log, TEXT("OCIOConfig already contains %s to %s-%s transform."), *InSourceColorSpace, *InDisplay, *InView);
+		return;
+	}
 
-		if (FoundSourceColorPtr == nullptr)
-		{
-			ColorTransforms.RemoveSingleSwap(Transform, true);
-			--i;
-			continue;
-		}
+	UOpenColorIOColorTransform* NewTransform = NewObject<UOpenColorIOColorTransform>(this, NAME_None, RF_NoFlags);
+	const bool bSuccess = NewTransform->Initialize(this, InSourceColorSpace, InDisplay, InView, bForwardDirection);
 
-		//The source was there so check if the destination color space of this transform is desired. If not, remove that transform from the list and move on.
-		const FOpenColorIOColorSpace* FoundDestinationColorPtr = DesiredColorSpaces.FindByPredicate([&](const FOpenColorIOColorSpace& InOtherColorSpace)
-		{
-			return InOtherColorSpace.ColorSpaceName == Transform->DestinationColorSpace;
-		});
+	if (bSuccess)
+	{
+		ColorTransforms.Add(NewTransform);
+	}
+	else
+	{
+		UE_LOG(LogOpenColorIO, Warning, TEXT("Could not create color space transform from %s to %s - %s. Verify your OCIO config file, it may have errors in it."), *InSourceColorSpace, *InDisplay, *InView);
+	}
+}
 
-		if (FoundDestinationColorPtr == nullptr)
+void UOpenColorIOConfiguration::CleanupTransforms()
+{
+	ColorTransforms = ColorTransforms.FilterByPredicate([this](const TObjectPtr<UOpenColorIOColorTransform>& InTransform) -> bool
 		{
-			ColorTransforms.RemoveSingleSwap(Transform, true);
-			--i;
-			continue;
-		}
-	}
+			//Check if the source color space of this transform is desired. If not, remove that transform from the list.
+			const FOpenColorIOColorSpace* FoundSourceColorPtr = DesiredColorSpaces.FindByPredicate([InTransform](const FOpenColorIOColorSpace& InOtherColorSpace)
+				{
+					return InOtherColorSpace.ColorSpaceName == InTransform->SourceColorSpace;
+				});
+
+			if (FoundSourceColorPtr == nullptr)
+			{
+				return false;
+			}
+
+			void* DestinationPtr = nullptr;
+
+			//The source was there so check if the destination color space/display view of this transform is desired. If not, remove that transform from the list.
+			if (InTransform->bIsDisplayViewType)
+			{
+				DestinationPtr = DesiredDisplayViews.FindByPredicate([InTransform](const FOpenColorIODisplayView& InOtherDisplayView)
+					{
+						return InOtherDisplayView.Display.Equals(InTransform->Display) && InOtherDisplayView.View.Equals(InTransform->View);
+					});
+
+			}
+			else
+			{
+				DestinationPtr = DesiredColorSpaces.FindByPredicate([InTransform](const FOpenColorIOColorSpace& InOtherColorSpace)
+					{
+						return InOtherColorSpace.ColorSpaceName == InTransform->DestinationColorSpace;
+					});
+			}
+
+			if (DestinationPtr == nullptr)
+			{
+				return false;
+			}
+
+			// Both source and destination are in the desired collections, so we keep the transform.
+			return true;
+		});
 }
 
 void UOpenColorIOConfiguration::PostLoad()
@@ -378,6 +458,31 @@ void UOpenColorIOConfiguration::PostEditChangeProperty(FPropertyChangedEvent& Pr
 					CreateColorTransform(TopColorSpace.ColorSpaceName, OtherColorSpace.ColorSpaceName);
 					CreateColorTransform(OtherColorSpace.ColorSpaceName, TopColorSpace.ColorSpaceName);
 				}
+
+				for (const FOpenColorIODisplayView& DisplayView : DesiredDisplayViews)
+				{
+					// note: we only support display-view transforms in the forward direction currently.
+					CreateColorTransform(TopColorSpace.ColorSpaceName, DisplayView.Display, DisplayView.View);
+				}
+			}
+		}
+
+		if (PropertyChangedEvent.ChangeType & (EPropertyChangeType::ArrayRemove | EPropertyChangeType::ArrayClear | EPropertyChangeType::ValueSet))
+		{
+			CleanupTransforms();
+		}
+	}
+	else if (PropertyChangedEvent.GetPropertyName() == GET_MEMBER_NAME_CHECKED(UOpenColorIOConfiguration, DesiredDisplayViews))
+	{
+		if (PropertyChangedEvent.ChangeType & (EPropertyChangeType::ArrayAdd | EPropertyChangeType::Duplicate | EPropertyChangeType::ValueSet))
+		{
+			for (const FOpenColorIODisplayView& DisplayView : DesiredDisplayViews)
+			{
+				for (const FOpenColorIOColorSpace& SourceColorSpace : DesiredColorSpaces)
+				{
+					// note: we only support display-view transforms in the forward direction currently.
+					CreateColorTransform(SourceColorSpace.ColorSpaceName, DisplayView.Display, DisplayView.View);
+				}
 			}
 		}
 
diff --git a/Engine/Plugins/Compositing/OpenColorIO/Source/OpenColorIO/Private/OpenColorIOShared.cpp b/Engine/Plugins/Compositing/OpenColorIO/Source/OpenColorIO/Private/OpenColorIOShared.cpp
--- a/Engine/Plugins/Compositing/OpenColorIO/Source/OpenColorIO/Private/OpenColorIOShared.cpp
+++ b/Engine/Plugins/Compositing/OpenColorIO/Source/OpenColorIO/Private/OpenColorIOShared.cpp
@@ -207,7 +207,7 @@ void FOpenColorIOTransformResource::SetupResource(ERHIFeatureLevel::Type InFeatu
 
 	ShaderCodeAndConfigHash = FullHash.ToString();
 	ShaderCode = InShadercode;
-	FriendlyName = InFriendlyName;
+	FriendlyName = InFriendlyName.Replace(TEXT("/"), TEXT(""));
 #if WITH_EDITOR
 	AssetPath = InAssetPath;
 #endif // WITH_EDITOR
diff --git a/Engine/Plugins/Compositing/OpenColorIO/Source/OpenColorIO/Public/OpenColorIOColorSpace.h b/Engine/Plugins/Compositing/OpenColorIO/Source/OpenColorIO/Public/OpenColorIOColorSpace.h
--- a/Engine/Plugins/Compositing/OpenColorIO/Source/OpenColorIO/Public/OpenColorIOColorSpace.h
+++ b/Engine/Plugins/Compositing/OpenColorIO/Source/OpenColorIO/Public/OpenColorIOColorSpace.h
@@ -69,6 +69,39 @@ public:
 	FString GetFamilyNameAtDepth(int32 InDepth) const;
 };
 
+USTRUCT(BlueprintType)
+struct OPENCOLORIO_API FOpenColorIODisplayView
+{
+	GENERATED_BODY()
+
+public:
+	/** Default constructor. */
+	FOpenColorIODisplayView();
+
+	/**
+	 * Create and initialize a new instance.
+	 */
+	FOpenColorIODisplayView(FStringView InDisplayName, FStringView InViewName);
+
+	UPROPERTY(VisibleAnywhere, Category = ColorSpace)
+	FString Display;
+
+	UPROPERTY(VisibleAnywhere, Category = ColorSpace)
+	FString View;
+
+	FString ToString() const;
+
+	/** Return true if the index and name have been set properly */
+	bool IsValid() const;
+
+	/** Reset members to default/empty values. */
+	void Reset();
+
+public:
+	bool operator==(const FOpenColorIODisplayView& Other) const { return Other.Display == Display && Other.View == View; }
+	bool operator!=(const FOpenColorIODisplayView& Other) const { return !operator==(Other); }
+};
+
 /**
  * Identifies a OCIO ColorSpace conversion.
  */
@@ -94,6 +127,10 @@ public:
 	UPROPERTY(EditAnywhere, BlueprintReadWrite, Category = ColorSpace)
 	FOpenColorIOColorSpace DestinationColorSpace;
 
+	/** The destination display view name. */
+	UPROPERTY(EditAnywhere, BlueprintReadWrite, Category = ColorSpace)
+	FOpenColorIODisplayView DestinationDisplayView;
+
 public:
 
 	/**
@@ -111,6 +148,10 @@ public:
 	* Ensure that the selected source and destination color spaces are valid, resets them otherwise.
 	*/
 	void ValidateColorSpaces();
+
+private:
+	/** Whether or not these settings are of the display-view type. */
+	bool IsDisplayView() const;
 };
 
 /**
diff --git a/Engine/Plugins/Compositing/OpenColorIO/Source/OpenColorIO/Public/OpenColorIOColorTransform.h b/Engine/Plugins/Compositing/OpenColorIO/Source/OpenColorIO/Public/OpenColorIOColorTransform.h
--- a/Engine/Plugins/Compositing/OpenColorIO/Source/OpenColorIO/Public/OpenColorIOColorTransform.h
+++ b/Engine/Plugins/Compositing/OpenColorIO/Source/OpenColorIO/Public/OpenColorIOColorTransform.h
@@ -33,6 +33,7 @@ public:
 
 public:
 	bool Initialize(UOpenColorIOConfiguration* InOwner, const FString& InSourceColorSpace, const FString& InDestinationColorSpace);
+	bool Initialize(UOpenColorIOConfiguration* InOwner, const FString& InSourceColorSpace, const FString& InDisplay, const FString& InView, bool bInForwardDirection = true);
 
 	/**
 	 * Cache resource shaders for cooking on the given shader platform.
@@ -71,6 +72,7 @@ public:
 	bool GetShaderAndLUTResouces(ERHIFeatureLevel::Type InFeatureLevel, FOpenColorIOTransformResource*& OutShaderResource, FTextureResource*& OutLUT3dResource);
 
 	bool IsTransform(const FString& InSourceColorSpace, const FString& InDestinationColorSpace) const;
+	bool IsTransform(const FString& InSourceColorSpace, const FString& InDisplay, const FString& InView, bool bInForwardDirection = true) const;
 
 	// For all ColorTransforms, UOpenColorIOColorTransform::CacheResourceShadersForRendering
 	static void AllColorTransformsCacheResourceShadersForRendering();
@@ -103,6 +105,7 @@ protected:
 	 * @return true if parameters are valid.
 	 */
 	bool GenerateColorTransformData(const FString& InSourceColorSpace, const FString& InDestinationColorSpace);
+	bool GenerateColorTransformData(const FString& InSourceColorSpace, const FString& InDisplay, const FString& InView, bool bInForwardDirection = true);
 
 	/**
 	 * Helper function returning the color space transform name based on source and destination color spaces.
@@ -162,12 +165,20 @@ public:
 	UPROPERTY()
 	TObjectPtr<UOpenColorIOConfiguration> ConfigurationOwner;
 
+	UPROPERTY(VisibleAnywhere, Category = "ColorSpace")
+	bool bIsDisplayViewType = false;
+
 	UPROPERTY(VisibleAnywhere, Category = "ColorSpace")
 	FString SourceColorSpace;
 
 	UPROPERTY(VisibleAnywhere, Category = "ColorSpace")
 	FString DestinationColorSpace;
 
+	UPROPERTY(VisibleAnywhere, Category = "ColorSpace")
+	FString Display;
+
+	UPROPERTY(VisibleAnywhere, Category = "ColorSpace")
+	FString View;
 private:
 
 	/** If the color space requires textures, this will contains the data to do the transform */
diff --git a/Engine/Plugins/Compositing/OpenColorIO/Source/OpenColorIO/Public/OpenColorIOConfiguration.h b/Engine/Plugins/Compositing/OpenColorIO/Source/OpenColorIO/Public/OpenColorIOConfiguration.h
--- a/Engine/Plugins/Compositing/OpenColorIO/Source/OpenColorIO/Public/OpenColorIOConfiguration.h
+++ b/Engine/Plugins/Compositing/OpenColorIO/Source/OpenColorIO/Public/OpenColorIOConfiguration.h
@@ -39,7 +39,9 @@ public:
 	UE_DEPRECATED(5.1, "GetShaderAndLUTResources is deprecated, please use GetRenderResources instead.")
 	bool GetShaderAndLUTResources(ERHIFeatureLevel::Type InFeatureLevel, const FString& InSourceColorSpace, const FString& InDestinationColorSpace, FOpenColorIOTransformResource*& OutShaderResource, FTextureResource*& OutLUT3dResource);
 	bool HasTransform(const FString& InSourceColorSpace, const FString& InDestinationColorSpace);
+	bool HasTransform(const FString& InSourceColorSpace, const FString& InDisplay, const FString& InView);
 	bool HasDesiredColorSpace(const FOpenColorIOColorSpace& ColorSpace) const;
+	bool HasDesiredDisplayView(const FOpenColorIODisplayView& DisplayView) const;
 	bool Validate() const;
 
 	/** This forces to reload colorspaces and corresponding shaders if those are not loaded already. */
@@ -57,6 +59,7 @@ public:
 
 protected:
 	void CreateColorTransform(const FString& InSourceColorSpace, const FString& InDestinationColorSpace);
+	void CreateColorTransform(const FString& InSourceColorSpace, const FString& InDisplay, const FString& InView, bool bForwardDirection = true);
 	void CleanupTransforms();
 
 	/** Same as above except user can specify the path manually. */
@@ -91,6 +94,9 @@ public:
 	UPROPERTY(BlueprintReadWrite, EditAnywhere, Category = "ColorSpace", meta=(OCIOConfigFile="ConfigurationFile"))
 	TArray<FOpenColorIOColorSpace> DesiredColorSpaces;
 
+	UPROPERTY(BlueprintReadWrite, EditAnywhere, Category = "ColorSpace", DisplayName="Desired Display-Views", meta = (OCIOConfigFile = "ConfigurationFile"))
+	TArray<FOpenColorIODisplayView> DesiredDisplayViews;
+
 private:
 
 	UPROPERTY()
diff --git a/Engine/Plugins/Compositing/OpenColorIO/Source/OpenColorIOEditor/Private/OpenColorIOColorSpaceConversionCustomization.cpp b/Engine/Plugins/Compositing/OpenColorIO/Source/OpenColorIOEditor/Private/OpenColorIOColorSpaceConversionCustomization.cpp
--- a/Engine/Plugins/Compositing/OpenColorIO/Source/OpenColorIOEditor/Private/OpenColorIOColorSpaceConversionCustomization.cpp
+++ b/Engine/Plugins/Compositing/OpenColorIO/Source/OpenColorIOEditor/Private/OpenColorIOColorSpaceConversionCustomization.cpp
@@ -14,6 +14,7 @@
 #include "PropertyHandle.h"
 #include "SResetToDefaultMenu.h"
 #include "Widgets/Input/SComboButton.h"
+#include "Widgets/Input/SSegmentedControl.h"
 
 #define LOCTEXT_NAMESPACE "OpenColorIOColorSpaceConversionCustomization"
 
@@ -26,6 +27,14 @@ void FOpenColorIOColorConversionSettingsCustomization::CustomizeHeader(TSharedRe
 		FProperty* Property = ColorConversionProperty->GetProperty();
 		check(Property && CastField<FStructProperty>(Property) && CastField<FStructProperty>(Property)->Struct && CastField<FStructProperty>(Property)->Struct->IsChildOf(FOpenColorIOColorConversionSettings::StaticStruct()));
 
+		TArray<void*> RawData;
+		ColorConversionProperty->AccessRawData(RawData);
+		if (RawData.Num() > 0)
+		{
+			FOpenColorIOColorConversionSettings* ColorSpaceConversion = reinterpret_cast<FOpenColorIOColorConversionSettings*>(RawData[0]);
+			bIsDestinationDisplayView = ColorSpaceConversion->DestinationDisplayView.IsValid();
+		}
+
 		if (ColorConversionProperty->GetNumPerObjectValues() == 1 && ColorConversionProperty->IsValidHandle())
 		{
 			TSharedPtr<IPropertyUtilities> PropertyUtils = CustomizationUtils.GetPropertyUtilities();
@@ -72,7 +81,9 @@ void FOpenColorIOColorConversionSettingsCustomization::CustomizeHeader(TSharedRe
 
 void FOpenColorIOColorConversionSettingsCustomization::CustomizeChildren(TSharedRef<IPropertyHandle> InStructPropertyHandle, IDetailChildrenBuilder& StructBuilder, IPropertyTypeCustomizationUtils& StructCustomizationUtils)
 {
-	TSharedPtr<IPropertyUtilities> PropertyUtils = StructCustomizationUtils.GetPropertyUtilities();
+
+	TAttribute<EVisibility> DestinationColorSpaceVisibility(this, &FOpenColorIOColorConversionSettingsCustomization::ShouldShowDestinationColorSpace);
+	TAttribute<EVisibility> DestinationDisplayViewVisibility(this, &FOpenColorIOColorConversionSettingsCustomization::ShouldShowDestinationDisplayView);
 
 	uint32 NumberOfChild;
 	if (InStructPropertyHandle->GetNumChildren(NumberOfChild) == FPropertyAccess::Success)
@@ -88,14 +99,25 @@ void FOpenColorIOColorConversionSettingsCustomization::CustomizeChildren(TShared
 				SourceColorSpaceProperty = ChildHandle;
 
 				FDetailWidgetRow& ColorSpaceWidget = StructBuilder.AddCustomRow(FText::FromName(ChildHandle->GetProperty()->GetFName()));
-				AddColorSpaceRow(ColorSpaceWidget, ChildHandle, StructCustomizationUtils);
+				AddPropertyRow(ColorSpaceWidget, ChildHandle, StructCustomizationUtils, false);
+
+				AddDestinationModeRow(StructBuilder, StructCustomizationUtils);
 			}
 			else if (ChildHandle->GetProperty()->GetFName() == GET_MEMBER_NAME_CHECKED(FOpenColorIOColorConversionSettings, DestinationColorSpace))
 			{
 				DestinationColorSpaceProperty = ChildHandle;
 
 				FDetailWidgetRow& ColorSpaceWidget = StructBuilder.AddCustomRow(FText::FromName(ChildHandle->GetProperty()->GetFName()));
-				AddColorSpaceRow(ColorSpaceWidget, ChildHandle, StructCustomizationUtils);
+				ColorSpaceWidget.Visibility(DestinationColorSpaceVisibility);
+				AddPropertyRow(ColorSpaceWidget, ChildHandle, StructCustomizationUtils, false);
+			}
+			else if (ChildHandle->GetProperty()->GetFName() == GET_MEMBER_NAME_CHECKED(FOpenColorIOColorConversionSettings, DestinationDisplayView))
+			{
+				DestinationDisplayViewProperty = ChildHandle;
+
+				FDetailWidgetRow& ColorSpaceWidget = StructBuilder.AddCustomRow(FText::FromName(ChildHandle->GetProperty()->GetFName()));
+				ColorSpaceWidget.Visibility(DestinationDisplayViewVisibility);
+				AddPropertyRow(ColorSpaceWidget, ChildHandle, StructCustomizationUtils, true);
 			}
 			else
 			{
@@ -106,7 +128,38 @@ void FOpenColorIOColorConversionSettingsCustomization::CustomizeChildren(TShared
 	InStructPropertyHandle->MarkHiddenByCustomization();
 }
 
-void FOpenColorIOColorConversionSettingsCustomization::AddColorSpaceRow(FDetailWidgetRow& InWidgetRow, TSharedRef<IPropertyHandle> InChildHandle, IPropertyTypeCustomizationUtils& InCustomizationUtils) const
+void FOpenColorIOColorConversionSettingsCustomization::AddDestinationModeRow(IDetailChildrenBuilder& StructBuilder, IPropertyTypeCustomizationUtils& StructCustomizationUtils)
+{
+	TSharedPtr<IPropertyUtilities> PropertyUtils = StructCustomizationUtils.GetPropertyUtilities();
+	StructBuilder.AddCustomRow(FText::FromString(TEXT("Destination Mode")))
+		.NameContent()
+		[
+			SNew(STextBlock)
+			.Text(LOCTEXT("DestinationMode", "Destination Mode"))
+		.ToolTipText(LOCTEXT("DestinationMode_Tooltip", "The destination mode to use, between color space and dislay-view."))
+		.Font(IPropertyTypeCustomizationUtils::GetRegularFont())
+		]
+		.ValueContent()
+		[
+			SNew(SSegmentedControl<bool>)
+			.Value_Lambda([this]()
+				{
+					return bIsDestinationDisplayView;
+				})
+		.OnValueChanged(this, &FOpenColorIOColorConversionSettingsCustomization::SetDestinationMode)
+			+ SSegmentedControl<bool>::Slot(false)
+			.Text(LOCTEXT("ColorSpace", "Color Space"))
+			.ToolTip(LOCTEXT("ColorSpace_ToolTip",
+				"Select this if you want to use a color space destination."))
+
+			+ SSegmentedControl<bool>::Slot(true)
+			.Text(LOCTEXT("DisplayView", "Display-View"))
+			.ToolTip(LOCTEXT("DisplayView_ToolTip",
+				"Select this if you want to use a display-view destination."))
+		];
+}
+
+void FOpenColorIOColorConversionSettingsCustomization::AddPropertyRow(FDetailWidgetRow& InWidgetRow, TSharedRef<IPropertyHandle> InChildHandle, IPropertyTypeCustomizationUtils& InCustomizationUtils, bool bIsDisplayViewRow) const
 {
 	TSharedPtr<IPropertyUtilities> PropertyUtils = InCustomizationUtils.GetPropertyUtilities();
 	TSharedPtr<SResetToDefaultMenu> ResetToDefaultMenu;
@@ -131,10 +184,21 @@ void FOpenColorIOColorConversionSettingsCustomization::AddColorSpaceRow(FDetailW
 					InChildHandle->AccessRawData(RawData);
 					if (RawData.Num() > 0)
 					{
-						FOpenColorIOColorSpace* ColorSpace = reinterpret_cast<FOpenColorIOColorSpace*>(RawData[0]);
-						if (ColorSpace != nullptr)
+						if (bIsDisplayViewRow)
+						{
+							FOpenColorIODisplayView* DisplayView = reinterpret_cast<FOpenColorIODisplayView*>(RawData[0]);
+							if (DisplayView != nullptr)
+							{
+								return FText::FromString(*DisplayView->ToString());
+							}
+						}
+						else
 						{
-							return FText::FromString(*ColorSpace->ToString());
+							FOpenColorIOColorSpace* ColorSpace = reinterpret_cast<FOpenColorIOColorSpace*>(RawData[0]);
+							if (ColorSpace != nullptr)
+							{
+								return FText::FromString(*ColorSpace->ToString());
+							}
 						}
 					}
 				
@@ -147,7 +211,7 @@ void FOpenColorIOColorConversionSettingsCustomization::AddColorSpaceRow(FDetailW
 			.VAlign(VAlign_Center)
 			[
 				SNew(SComboButton)
-				.OnGetMenuContent_Lambda([=]() { return HandleColorSpaceComboButtonMenuContent(InChildHandle); })
+				.OnGetMenuContent_Lambda([=]() { return bIsDisplayViewRow ? HandleDisplayViewComboButtonMenuContent(InChildHandle) : HandleColorSpaceComboButtonMenuContent(InChildHandle); })
 				.ContentPadding(FMargin(4.0, 2.0))
 			]
 			+ SHorizontalBox::Slot()
@@ -157,7 +221,7 @@ void FOpenColorIOColorConversionSettingsCustomization::AddColorSpaceRow(FDetailW
 				SAssignNew(ResetToDefaultMenu, SResetToDefaultMenu)
 			]
 		].IsEnabled(MakeAttributeLambda([=] { return !InChildHandle->IsEditConst() && PropertyUtils->IsPropertyEditingEnabled(); }));
-	
+		
 		ResetToDefaultMenu->AddProperty(InChildHandle);
 }
 
@@ -190,7 +254,7 @@ TSharedRef<SWidget> FOpenColorIOColorConversionSettingsCustomization::HandleColo
 		// generate menu
 		FMenuBuilder MenuBuilder(true, nullptr);
 
-		MenuBuilder.BeginSection("AvailableColorSpaces", LOCTEXT("AvailableCoorSpaces", "Available Color Spaces"));
+		MenuBuilder.BeginSection("AvailableColorSpaces", LOCTEXT("AvailableColorSpaces", "Available Color Spaces"));
 		{
 			bool ColorSpaceAdded = false;
 
@@ -265,4 +329,121 @@ TSharedRef<SWidget> FOpenColorIOColorConversionSettingsCustomization::HandleColo
 	return SNullWidget::NullWidget;
 }
 
+TSharedRef<SWidget> FOpenColorIOColorConversionSettingsCustomization::HandleDisplayViewComboButtonMenuContent(TSharedPtr<IPropertyHandle> InPropertyHandle) const
+{
+	FOpenColorIOColorConversionSettings* ColorSpaceConversion = nullptr;
+
+	if (InPropertyHandle.IsValid())
+	{
+		TArray<void*> RawData;
+		ColorConversionProperty->AccessRawData(RawData);
+		if (RawData.Num() > 0)
+		{
+			ColorSpaceConversion = reinterpret_cast<FOpenColorIOColorConversionSettings*>(RawData[0]);
+		}
+	}
+
+	if (ColorSpaceConversion && ColorSpaceConversion->ConfigurationSource)
+	{
+		// generate menu
+		FMenuBuilder MenuBuilder(true, nullptr);
+
+		MenuBuilder.BeginSection("AvailableDisplayViews", LOCTEXT("AvailableDisplayViews", "Available Display-Views"));
+		{
+			bool DisplayViewAdded = false;
+
+			for (int32 i = 0; i < ColorSpaceConversion->ConfigurationSource->DesiredDisplayViews.Num(); ++i)
+			{
+				const FOpenColorIODisplayView& DisplayView = ColorSpaceConversion->ConfigurationSource->DesiredDisplayViews[i];
+				if (!DisplayView.IsValid())
+				{
+					continue;
+				}
+
+				MenuBuilder.AddMenuEntry
+				(
+					FText::FromString(DisplayView.ToString()),
+					FText::FromString(DisplayView.ToString()),
+					FSlateIcon(),
+					FUIAction
+					(
+						FExecuteAction::CreateLambda([=]
+							{
+								if (InPropertyHandle.IsValid())
+								{
+									if (FStructProperty* StructProperty = CastField<FStructProperty>(InPropertyHandle->GetProperty()))
+									{
+										TArray<void*> RawData;
+										InPropertyHandle->AccessRawData(RawData);
+										if (RawData.Num() > 0)
+										{
+											FOpenColorIODisplayView* PreviousDisplayViewValue = reinterpret_cast<FOpenColorIODisplayView*>(RawData[0]);
+
+											FString TextValue;
+											StructProperty->Struct->ExportText(TextValue, &DisplayView, PreviousDisplayViewValue, nullptr, EPropertyPortFlags::PPF_None, nullptr);
+											ensure(InPropertyHandle->SetValueFromFormattedString(TextValue, EPropertyValueSetFlags::DefaultFlags) == FPropertyAccess::Result::Success);
+										}
+									}
+								}
+							}),
+						FCanExecuteAction(),
+						FIsActionChecked::CreateLambda([=]
+							{
+								return false;
+							})
+						),
+						NAME_None,
+						EUserInterfaceActionType::RadioButton
+					);
+
+				DisplayViewAdded = true;
+			}
+
+			if (!DisplayViewAdded)
+			{
+				MenuBuilder.AddWidget(SNullWidget::NullWidget, LOCTEXT("NoDisplayViewFound", "No available display-view"), false, false);
+			}
+		}
+		MenuBuilder.EndSection();
+
+		return MenuBuilder.MakeWidget();
+	}
+
+	return SNullWidget::NullWidget;
+}
+
+void FOpenColorIOColorConversionSettingsCustomization::SetDestinationMode(bool bInIsDestinationDisplayView)
+{
+	bIsDestinationDisplayView = bInIsDestinationDisplayView;
+
+	if (ColorConversionProperty.IsValid())
+	{
+		TArray<void*> RawData;
+		ColorConversionProperty->AccessRawData(RawData);
+		if (RawData.Num() > 0)
+		{
+			FOpenColorIOColorConversionSettings* ColorSpaceConversion = reinterpret_cast<FOpenColorIOColorConversionSettings*>(RawData[0]);
+
+			if (bIsDestinationDisplayView)
+			{
+				ColorSpaceConversion->DestinationColorSpace.Reset();
+			}
+			else
+			{
+				ColorSpaceConversion->DestinationDisplayView.Reset();
+			}
+		}
+	}
+}
+
+EVisibility FOpenColorIOColorConversionSettingsCustomization::ShouldShowDestinationColorSpace() const
+{
+	return bIsDestinationDisplayView ? EVisibility::Hidden : EVisibility::Visible;
+}
+
+EVisibility FOpenColorIOColorConversionSettingsCustomization::ShouldShowDestinationDisplayView() const
+{
+	return bIsDestinationDisplayView ? EVisibility::Visible : EVisibility::Hidden;
+}
+
 #undef LOCTEXT_NAMESPACE
diff --git a/Engine/Plugins/Compositing/OpenColorIO/Source/OpenColorIOEditor/Private/OpenColorIOColorSpaceConversionCustomization.h b/Engine/Plugins/Compositing/OpenColorIO/Source/OpenColorIOEditor/Private/OpenColorIOColorSpaceConversionCustomization.h
--- a/Engine/Plugins/Compositing/OpenColorIO/Source/OpenColorIOEditor/Private/OpenColorIOColorSpaceConversionCustomization.h
+++ b/Engine/Plugins/Compositing/OpenColorIO/Source/OpenColorIOEditor/Private/OpenColorIOColorSpaceConversionCustomization.h
@@ -22,8 +22,18 @@ public:
 	virtual void CustomizeChildren(TSharedRef<class IPropertyHandle> InPropertyHandle, class IDetailChildrenBuilder& StructBuilder, IPropertyTypeCustomizationUtils& PropertyTypeCustomizationUtils) override;
 
 private:
-	void AddColorSpaceRow(FDetailWidgetRow& InWidgetRow, TSharedRef<IPropertyHandle> InChildHandle, IPropertyTypeCustomizationUtils& InCustomizationUtils) const;
+	void AddDestinationModeRow(IDetailChildrenBuilder& StructBuilder, IPropertyTypeCustomizationUtils& StructCustomizationUtils);
+	void AddPropertyRow(FDetailWidgetRow& InWidgetRow, TSharedRef<IPropertyHandle> InChildHandle, IPropertyTypeCustomizationUtils& InCustomizationUtils, bool bIsDisplayViewRow) const;
+	
 	TSharedRef<SWidget> HandleColorSpaceComboButtonMenuContent(TSharedPtr<IPropertyHandle> InPropertyHandle) const;
+	TSharedRef<SWidget> HandleDisplayViewComboButtonMenuContent(TSharedPtr<IPropertyHandle> InPropertyHandle) const;
+
+	void SetDestinationMode(bool bInIsDestinationDisplayView);
+	/** Controls visibility for widgets when the destination is a color space. */
+	EVisibility ShouldShowDestinationColorSpace() const;
+
+	/** Controls visibility for widgets when the destination is a display-view. */
+	EVisibility ShouldShowDestinationDisplayView() const;
 
 	/** Pointer to the ColorConversion struct property handle. */
 	TSharedPtr<IPropertyHandle> ColorConversionProperty;
@@ -33,4 +43,9 @@ private:
 	
 	/** Pointer to the ColorConversion struct member DestinationColorSpace property handle. */
 	TSharedPtr<IPropertyHandle> DestinationColorSpaceProperty;
+
+	/** Pointer to the ColorConversion struct member DestinationColorSpace property handle. */
+	TSharedPtr<IPropertyHandle> DestinationDisplayViewProperty;
+
+	bool bIsDestinationDisplayView = false;
 };
diff --git a/Engine/Plugins/Compositing/OpenColorIO/Source/OpenColorIOEditor/Private/OpenColorIOColorSpaceCustomization.cpp b/Engine/Plugins/Compositing/OpenColorIO/Source/OpenColorIOEditor/Private/OpenColorIOColorSpaceCustomization.cpp
--- a/Engine/Plugins/Compositing/OpenColorIO/Source/OpenColorIOEditor/Private/OpenColorIOColorSpaceCustomization.cpp
+++ b/Engine/Plugins/Compositing/OpenColorIO/Source/OpenColorIOEditor/Private/OpenColorIOColorSpaceCustomization.cpp
@@ -17,15 +17,15 @@ void FOpenColorIOColorSpaceCustomization::CustomizeHeader(TSharedRef<IPropertyHa
 	CachedConfigFile.reset();
 	LoadedFilePath = FFilePath();
 	
-	ColorSpaceProperty = InPropertyHandle;
+	CachedProperty = InPropertyHandle;
 
-	if (ColorSpaceProperty->GetNumPerObjectValues() == 1 && ColorSpaceProperty->IsValidHandle())
+	if (CachedProperty->GetNumPerObjectValues() == 1 && CachedProperty->IsValidHandle())
 	{
-		FProperty* Property = ColorSpaceProperty->GetProperty();
+		FProperty* Property = CachedProperty->GetProperty();
 		check(Property && CastField<FStructProperty>(Property) && CastField<FStructProperty>(Property)->Struct && CastField<FStructProperty>(Property)->Struct->IsChildOf(FOpenColorIOColorSpace::StaticStruct()));
 
 		TArray<void*> RawData;
-		ColorSpaceProperty->AccessRawData(RawData);
+		CachedProperty->AccessRawData(RawData);
 
 		check(RawData.Num() == 1);
 		FOpenColorIOColorSpace* ColorSpaceValue = reinterpret_cast<FOpenColorIOColorSpace*>(RawData[0]);
@@ -34,39 +34,7 @@ void FOpenColorIOColorSpaceCustomization::CustomizeHeader(TSharedRef<IPropertyHa
 		TSharedPtr<IPropertyUtilities> PropertyUtils = CustomizationUtils.GetPropertyUtilities();
 
 		// Get the ConfigurationFile to read color space from
-		{
-			static FName NAME_ConfigFile = TEXT("OCIOConfigFile");
-
-			//Verify if we're in an array before accessing the property directly. Need to go deeper for array properties.
-			TSharedPtr<IPropertyHandle> ParentArrayHandle = InPropertyHandle->GetParentHandle();
-			TSharedPtr<IPropertyHandleArray> ParentHandleAsArray = ParentArrayHandle->AsArray();
-			if (ParentHandleAsArray.IsValid())
-			{
-				//Struct is in an array, parent class is access through more layers
-				if (ParentArrayHandle->GetProperty()->HasMetaData(NAME_ConfigFile))
-				{
-					const FString& ConfigFileVariableName = ParentArrayHandle->GetProperty()->GetMetaData(NAME_ConfigFile);
-					TSharedPtr<IPropertyHandle> ClassHandle = ParentArrayHandle->GetParentHandle()->GetParentHandle();
-					if (ClassHandle.IsValid() && !ConfigFileVariableName.IsEmpty())
-					{
-						ConfigurationFileProperty = ClassHandle->GetChildHandle(*ConfigFileVariableName);
-					}
-				}
-			}
-			else
-			{
-				if (ColorSpaceProperty->HasMetaData(NAME_ConfigFile))
-				{
-					const FString& ConfigFileVariableName = ColorSpaceProperty->GetMetaData(NAME_ConfigFile);
-					TSharedPtr<IPropertyHandle> ClassHandle = ColorSpaceProperty->GetParentHandle();
-					if (ClassHandle.IsValid() && !ConfigFileVariableName.IsEmpty())
-					{
-						ConfigurationFileProperty = ClassHandle->GetChildHandle(*ConfigFileVariableName);
-					}
-				}
-			}
-		}
-
+		ConfigurationFileProperty = GetConfigurationFileProperty();
 
 		HeaderRow
 			.NameContent()
@@ -97,12 +65,101 @@ void FOpenColorIOColorSpaceCustomization::CustomizeHeader(TSharedRef<IPropertyHa
 	}
 }
 
-void FOpenColorIOColorSpaceCustomization::CustomizeChildren(TSharedRef<IPropertyHandle> InStructPropertyHandle, IDetailChildrenBuilder& StructBuilder, IPropertyTypeCustomizationUtils& StructCustomizationUtils)
+void FOpenColorIODisplayViewCustomization::CustomizeHeader(TSharedRef<IPropertyHandle> InPropertyHandle, FDetailWidgetRow& HeaderRow, IPropertyTypeCustomizationUtils& CustomizationUtils)
 {
+	//Reset internals
+	ConfigurationFileProperty.Reset();
+	CachedConfigFile.reset();
+	LoadedFilePath = FFilePath();
+
+	CachedProperty = InPropertyHandle;
+
+	if (CachedProperty->GetNumPerObjectValues() == 1 && CachedProperty->IsValidHandle())
+	{
+		FProperty* Property = CachedProperty->GetProperty();
+		check(Property && CastField<FStructProperty>(Property) && CastField<FStructProperty>(Property)->Struct && CastField<FStructProperty>(Property)->Struct->IsChildOf(FOpenColorIODisplayView::StaticStruct()));
+
+		TArray<void*> RawData;
+		CachedProperty->AccessRawData(RawData);
+
+		check(RawData.Num() == 1);
+		FOpenColorIODisplayView* DisplayViewValue = reinterpret_cast<FOpenColorIODisplayView*>(RawData[0]);
+
+		check(DisplayViewValue);
+		TSharedPtr<IPropertyUtilities> PropertyUtils = CustomizationUtils.GetPropertyUtilities();
+
+		// Get the ConfigurationFile to read color space from
+		ConfigurationFileProperty = GetConfigurationFileProperty();
+
+		HeaderRow
+			.NameContent()
+			[
+				InPropertyHandle->CreatePropertyNameWidget()
+			]
+		.ValueContent()
+			.MaxDesiredWidth(512)
+			[
+				SNew(SHorizontalBox)
+				+ SHorizontalBox::Slot()
+				.FillWidth(1.0f)
+				.VAlign(VAlign_Center)
+				[
+					SNew(STextBlock)
+					.Text(MakeAttributeLambda([=] { return FText::FromString(DisplayViewValue->ToString()); }))
+				]
+				+ SHorizontalBox::Slot()
+				.AutoWidth()
+				.Padding(FMargin(4.0f, 0.0f, 0.0f, 0.0f))
+				.VAlign(VAlign_Center)
+				[
+					SNew(SComboButton)
+					.OnGetMenuContent(this, &FOpenColorIODisplayViewCustomization::HandleSourceComboButtonMenuContent)
+					.ContentPadding(FMargin(4.0, 2.0))
+				]
+			].IsEnabled(MakeAttributeLambda([=] { return !InPropertyHandle->IsEditConst() && PropertyUtils->IsPropertyEditingEnabled(); }));
+	}
 }
 
+TSharedPtr<IPropertyHandle> IPropertyTypeCustomizationOpenColorIO::GetConfigurationFileProperty() const
+{
+	if (CachedProperty.IsValid())
+	{
+		static FName NAME_ConfigFile = TEXT("OCIOConfigFile");
 
-bool FOpenColorIOColorSpaceCustomization::LoadConfigurationFile(const FFilePath& InFilePath)
+		//Verify if we're in an array before accessing the property directly. Need to go deeper for array properties.
+		TSharedPtr<IPropertyHandle> ParentArrayHandle = CachedProperty->GetParentHandle();
+		TSharedPtr<IPropertyHandleArray> ParentHandleAsArray = ParentArrayHandle->AsArray();
+		if (ParentHandleAsArray.IsValid())
+		{
+			//Struct is in an array, parent class is access through more layers
+			if (ParentArrayHandle->GetProperty()->HasMetaData(NAME_ConfigFile))
+			{
+				const FString& ConfigFileVariableName = ParentArrayHandle->GetProperty()->GetMetaData(NAME_ConfigFile);
+				TSharedPtr<IPropertyHandle> ClassHandle = ParentArrayHandle->GetParentHandle()->GetParentHandle();
+				if (ClassHandle.IsValid() && !ConfigFileVariableName.IsEmpty())
+				{
+					return ClassHandle->GetChildHandle(*ConfigFileVariableName);
+				}
+			}
+		}
+		else
+		{
+			if (CachedProperty->HasMetaData(NAME_ConfigFile))
+			{
+				const FString& ConfigFileVariableName = CachedProperty->GetMetaData(NAME_ConfigFile);
+				TSharedPtr<IPropertyHandle> ClassHandle = CachedProperty->GetParentHandle();
+				if (ClassHandle.IsValid() && !ConfigFileVariableName.IsEmpty())
+				{
+					return ClassHandle->GetChildHandle(*ConfigFileVariableName);
+				}
+			}
+		}
+	}
+
+	return nullptr;
+}
+
+bool IPropertyTypeCustomizationOpenColorIO::LoadConfigurationFile(const FFilePath& InFilePath)
 {
 #if WITH_OCIO
 #if !PLATFORM_EXCEPTIONS_DISABLED
@@ -145,6 +202,26 @@ bool FOpenColorIOColorSpaceCustomization::LoadConfigurationFile(const FFilePath&
 #endif //WITH_OCIO
 }
 
+bool IPropertyTypeCustomizationOpenColorIO::CheckValidConfiguration()
+{
+	if (ConfigurationFileProperty.IsValid())
+	{
+		TArray<void*> RawData;
+		ConfigurationFileProperty->AccessRawData(RawData);
+
+		check(RawData.Num() == 1);
+		FFilePath* ConfigFilePath = reinterpret_cast<FFilePath*>(RawData[0]);
+		check(ConfigFilePath);
+
+		if (!ConfigFilePath->FilePath.IsEmpty() && ConfigFilePath->FilePath != LoadedFilePath.FilePath)
+		{
+			return LoadConfigurationFile(*ConfigFilePath);
+		}
+	}
+
+	return false;
+}
+
 void FOpenColorIOColorSpaceCustomization::ProcessColorSpaceForMenuGeneration(FMenuBuilder& InMenuBuilder, const int32 InMenuDepth, const FString& InPreviousFamilyHierarchy, const FOpenColorIOColorSpace& InColorSpace, TArray<FString>& InOutExistingMenuFilter)
 {
 	const FString NextFamilyName = InColorSpace.GetFamilyNameAtDepth(InMenuDepth);
@@ -212,15 +289,15 @@ void FOpenColorIOColorSpaceCustomization::AddMenuEntry(FMenuBuilder& InMenuBuild
 		FUIAction(
 			FExecuteAction::CreateLambda([=]
 			{
-				if (FStructProperty* StructProperty = CastField<FStructProperty>(ColorSpaceProperty->GetProperty()))
+				if (FStructProperty* StructProperty = CastField<FStructProperty>(CachedProperty->GetProperty()))
 				{
 					TArray<void*> RawData;
-					ColorSpaceProperty->AccessRawData(RawData);
+					CachedProperty->AccessRawData(RawData);
 					FOpenColorIOColorSpace* PreviousColorSpaceValue = reinterpret_cast<FOpenColorIOColorSpace*>(RawData[0]);
 
 					FString TextValue;
 					StructProperty->Struct->ExportText(TextValue, &InColorSpace, PreviousColorSpaceValue, nullptr, EPropertyPortFlags::PPF_None, nullptr);
-					ensure(ColorSpaceProperty->SetValueFromFormattedString(TextValue, EPropertyValueSetFlags::DefaultFlags) == FPropertyAccess::Result::Success);
+					ensure(CachedProperty->SetValueFromFormattedString(TextValue, EPropertyValueSetFlags::DefaultFlags) == FPropertyAccess::Result::Success);
 				}
 			}
 			),
@@ -228,7 +305,7 @@ void FOpenColorIOColorSpaceCustomization::AddMenuEntry(FMenuBuilder& InMenuBuild
 			FIsActionChecked::CreateLambda([=]
 			{
 				TArray<void*> RawData;
-				ColorSpaceProperty->AccessRawData(RawData);
+				CachedProperty->AccessRawData(RawData);
 				FOpenColorIOColorSpace* ColorSpaceValue = reinterpret_cast<FOpenColorIOColorSpace*>(RawData[0]);
 				return *ColorSpaceValue == InColorSpace;
 			})
@@ -240,21 +317,7 @@ void FOpenColorIOColorSpaceCustomization::AddMenuEntry(FMenuBuilder& InMenuBuild
 
 TSharedRef<SWidget> FOpenColorIOColorSpaceCustomization::HandleSourceComboButtonMenuContent()
 {
-	bool bValidConfiguration = false;
-	if (ConfigurationFileProperty.IsValid())
-	{
-		TArray<void*> RawData;
-		ConfigurationFileProperty->AccessRawData(RawData);
-
-		check(RawData.Num() == 1);
-		FFilePath* ConfigFilePath = reinterpret_cast<FFilePath*>(RawData[0]);
-		check(ConfigFilePath);
-
-		if (!ConfigFilePath->FilePath.IsEmpty() && ConfigFilePath->FilePath != LoadedFilePath.FilePath)
-		{
-			bValidConfiguration = LoadConfigurationFile(*ConfigFilePath);
-		}
-	}
+	bool bValidConfiguration = CheckValidConfiguration();
 
 	// Generate menu
 	FMenuBuilder MenuBuilder(true, nullptr);
@@ -300,4 +363,92 @@ TSharedRef<SWidget> FOpenColorIOColorSpaceCustomization::HandleSourceComboButton
 	return MenuBuilder.MakeWidget();
 }
 
+void FOpenColorIODisplayViewCustomization::PopulateViewSubMenu(FMenuBuilder& InMenuBuilder, FOpenColorIODisplayView InDisplayView)
+{
+	const int32 ViewCount = CachedConfigFile->getNumViews(TCHAR_TO_ANSI(*InDisplayView.Display));
+	for (int32 i = 0; i < ViewCount; ++i)
+	{
+		const char* ViewName = CachedConfigFile->getView(TCHAR_TO_ANSI(*InDisplayView.Display), i);
+		InDisplayView.View = StringCast<TCHAR>(ViewName).Get();
+		
+		AddMenuEntry(InMenuBuilder, InDisplayView);
+	}
+}
+
+void FOpenColorIODisplayViewCustomization::AddMenuEntry(FMenuBuilder& InMenuBuilder, const FOpenColorIODisplayView& InDisplayView)
+{
+	InMenuBuilder.AddMenuEntry(
+		FText::FromString(InDisplayView.View),
+		FText::FromString(InDisplayView.ToString()),
+		FSlateIcon(),
+		FUIAction(
+			FExecuteAction::CreateLambda([=]
+				{
+					if (FStructProperty* StructProperty = CastField<FStructProperty>(CachedProperty->GetProperty()))
+					{
+						TArray<void*> RawData;
+						CachedProperty->AccessRawData(RawData);
+						FOpenColorIODisplayView* PreviousDisplayViewValue = reinterpret_cast<FOpenColorIODisplayView*>(RawData[0]);
+
+						FString TextValue;
+						StructProperty->Struct->ExportText(TextValue, &InDisplayView, PreviousDisplayViewValue, nullptr, EPropertyPortFlags::PPF_None, nullptr);
+						ensure(CachedProperty->SetValueFromFormattedString(TextValue, EPropertyValueSetFlags::DefaultFlags) == FPropertyAccess::Result::Success);
+					}
+				}
+			),
+			FCanExecuteAction(),
+			FIsActionChecked::CreateLambda([=]
+				{
+					TArray<void*> RawData;
+					CachedProperty->AccessRawData(RawData);
+					FOpenColorIODisplayView* DisplayViewValue = reinterpret_cast<FOpenColorIODisplayView*>(RawData[0]);
+					return *DisplayViewValue == InDisplayView;
+				})
+			),
+			NAME_None,
+			EUserInterfaceActionType::RadioButton
+		);
+}
+
+TSharedRef<SWidget> FOpenColorIODisplayViewCustomization::HandleSourceComboButtonMenuContent()
+{
+	bool bValidConfiguration = CheckValidConfiguration();
+
+	// Generate menu
+	FMenuBuilder MenuBuilder(true, nullptr);
+	TArray<FString> ExistingSubMenus;
+
+	MenuBuilder.BeginSection("AllDisplayViews", LOCTEXT("AllDisplayViewsSection", "Display - View"));
+	{
+		if (bValidConfiguration)
+		{
+			const int32 DisplayCount = CachedConfigFile->getNumDisplays();
+			for (int32 i = 0; i < DisplayCount; ++i)
+			{
+				const char* DisplayName = CachedConfigFile->getDisplay(i);
+
+				const FOpenColorIODisplayView DisplayViewValue = { StringCast<TCHAR>(DisplayName).Get(), TEXT("<Invalid>") };
+
+				MenuBuilder.AddSubMenu(
+					FText::FromString(DisplayViewValue.Display),
+					LOCTEXT("OpensDisplayViewSubMenu", "Display - View Family Sub Menu"),
+					FNewMenuDelegate::CreateRaw(this, &FOpenColorIODisplayViewCustomization::PopulateViewSubMenu, DisplayViewValue)
+				);
+			}
+
+			if (DisplayCount <= 0)
+			{
+				MenuBuilder.AddWidget(SNullWidget::NullWidget, LOCTEXT("NoDisplayFound", "No display found"), false, false);
+			}
+		}
+		else
+		{
+			MenuBuilder.AddWidget(SNullWidget::NullWidget, LOCTEXT("InvalidConfigurationFile", "Invalid configuration file"), false, false);
+		}
+	}
+	MenuBuilder.EndSection();
+
+	return MenuBuilder.MakeWidget();
+}
+
 #undef LOCTEXT_NAMESPACE
diff --git a/Engine/Plugins/Compositing/OpenColorIO/Source/OpenColorIOEditor/Private/OpenColorIOColorSpaceCustomization.h b/Engine/Plugins/Compositing/OpenColorIO/Source/OpenColorIOEditor/Private/OpenColorIOColorSpaceCustomization.h
--- a/Engine/Plugins/Compositing/OpenColorIO/Source/OpenColorIOEditor/Private/OpenColorIOColorSpaceCustomization.h
+++ b/Engine/Plugins/Compositing/OpenColorIO/Source/OpenColorIOEditor/Private/OpenColorIOColorSpaceCustomization.h
@@ -14,31 +14,24 @@
 /**
  * Implements a details view customization for the FOpenColorIOConfiguration
  */
-class FOpenColorIOColorSpaceCustomization : public IPropertyTypeCustomization
+class IPropertyTypeCustomizationOpenColorIO : public IPropertyTypeCustomization
 {
 public:
-	static TSharedRef<IPropertyTypeCustomization> MakeInstance()
-	{
-		return MakeShareable(new FOpenColorIOColorSpaceCustomization);
-	}
 
 	/** IPropertyTypeCustomization interface */
-	virtual void CustomizeHeader(TSharedRef<class IPropertyHandle> InPropertyHandle, class FDetailWidgetRow& HeaderRow, IPropertyTypeCustomizationUtils& PropertyTypeCustomizationUtils) override;
-	virtual void CustomizeChildren(TSharedRef<class IPropertyHandle> InPropertyHandle, class IDetailChildrenBuilder& StructBuilder, IPropertyTypeCustomizationUtils& PropertyTypeCustomizationUtils) override;
+	virtual void CustomizeHeader(TSharedRef<class IPropertyHandle> InPropertyHandle, class FDetailWidgetRow& HeaderRow, IPropertyTypeCustomizationUtils& PropertyTypeCustomizationUtils) override {}
+	virtual void CustomizeChildren(TSharedRef<class IPropertyHandle> InPropertyHandle, class IDetailChildrenBuilder& StructBuilder, IPropertyTypeCustomizationUtils& PropertyTypeCustomizationUtils) override {}
 
 
 protected:
+	TSharedPtr<IPropertyHandle> GetConfigurationFileProperty() const;
 	bool LoadConfigurationFile(const FFilePath& InFilePath);
+	bool CheckValidConfiguration();
 
-	void ProcessColorSpaceForMenuGeneration(FMenuBuilder& InMenuBuilder, const int32 InMenuDepth, const FString& InPreviousFamilyHierarchy, const FOpenColorIOColorSpace& InColorSpace, TArray<FString>& InOutExistingMenuFilter);
-	void PopulateSubMenu(FMenuBuilder& InMenuBuilder, const int32 InMenuDepth, FString InPreviousFamilyHierarchy);
-	void AddMenuEntry(FMenuBuilder& InMenuBuilder, const FOpenColorIOColorSpace& InColorSpace);
-
-private:
-	TSharedRef<SWidget> HandleSourceComboButtonMenuContent();
+protected:
 
-	/** Pointer to the ColorSpace property handle. */
-	TSharedPtr<IPropertyHandle> ColorSpaceProperty;
+	/** Pointer to the property handle. */
+	TSharedPtr<IPropertyHandle> CachedProperty;
 
 	/** Pointer to the ConfigurationFile property handle. */
 	TSharedPtr<IPropertyHandle> ConfigurationFileProperty;
@@ -50,3 +43,42 @@ private:
 	OCIO_NAMESPACE::ConstConfigRcPtr CachedConfigFile;
 };
 
+
+class FOpenColorIOColorSpaceCustomization : public IPropertyTypeCustomizationOpenColorIO
+{
+public:
+	static TSharedRef<IPropertyTypeCustomization> MakeInstance()
+	{
+		return MakeShareable(new FOpenColorIOColorSpaceCustomization);
+	}
+
+	virtual void CustomizeHeader(TSharedRef<class IPropertyHandle> InPropertyHandle, class FDetailWidgetRow& HeaderRow, IPropertyTypeCustomizationUtils& PropertyTypeCustomizationUtils) override;
+
+private:
+	void ProcessColorSpaceForMenuGeneration(FMenuBuilder& InMenuBuilder, const int32 InMenuDepth, const FString& InPreviousFamilyHierarchy, const FOpenColorIOColorSpace& InColorSpace, TArray<FString>& InOutExistingMenuFilter);
+	void PopulateSubMenu(FMenuBuilder& InMenuBuilder, const int32 InMenuDepth, FString InPreviousFamilyHierarchy);
+	void AddMenuEntry(FMenuBuilder& InMenuBuilder, const FOpenColorIOColorSpace& InColorSpace);
+
+	TSharedRef<SWidget> HandleSourceComboButtonMenuContent();
+};
+
+/**
+ * Implements a details view customization for the FOpenColorIOConfiguration
+ */
+class FOpenColorIODisplayViewCustomization : public IPropertyTypeCustomizationOpenColorIO
+{
+public:
+	static TSharedRef<IPropertyTypeCustomization> MakeInstance()
+	{
+		return MakeShareable(new FOpenColorIODisplayViewCustomization);
+	}
+
+	virtual void CustomizeHeader(TSharedRef<class IPropertyHandle> InPropertyHandle, class FDetailWidgetRow& HeaderRow, IPropertyTypeCustomizationUtils& PropertyTypeCustomizationUtils) override;
+
+private:
+	void PopulateViewSubMenu(FMenuBuilder& InMenuBuilder, FOpenColorIODisplayView InDisplayView);
+	void AddMenuEntry(FMenuBuilder& InMenuBuilder, const FOpenColorIODisplayView& InDisplayView);
+
+	TSharedRef<SWidget> HandleSourceComboButtonMenuContent();
+};
+
diff --git a/Engine/Plugins/Compositing/OpenColorIO/Source/OpenColorIOEditor/Private/OpenColorIOEditorModule.cpp b/Engine/Plugins/Compositing/OpenColorIO/Source/OpenColorIOEditor/Private/OpenColorIOEditorModule.cpp
--- a/Engine/Plugins/Compositing/OpenColorIO/Source/OpenColorIOEditor/Private/OpenColorIOEditorModule.cpp
+++ b/Engine/Plugins/Compositing/OpenColorIO/Source/OpenColorIOEditor/Private/OpenColorIOEditorModule.cpp
@@ -79,6 +79,7 @@ void FOpenColorIOEditorModule::RegisterCustomizations()
 	FPropertyEditorModule& PropertyModule = FModuleManager::LoadModuleChecked<FPropertyEditorModule>("PropertyEditor");
 	PropertyModule.RegisterCustomPropertyTypeLayout(FOpenColorIOColorConversionSettings::StaticStruct()->GetFName(), FOnGetPropertyTypeCustomizationInstance::CreateStatic(&FOpenColorIOColorConversionSettingsCustomization::MakeInstance));
 	PropertyModule.RegisterCustomPropertyTypeLayout(FOpenColorIOColorSpace::StaticStruct()->GetFName(), FOnGetPropertyTypeCustomizationInstance::CreateStatic(&FOpenColorIOColorSpaceCustomization::MakeInstance));
+	PropertyModule.RegisterCustomPropertyTypeLayout(FOpenColorIODisplayView::StaticStruct()->GetFName(), FOnGetPropertyTypeCustomizationInstance::CreateStatic(&FOpenColorIODisplayViewCustomization::MakeInstance));
 	PropertyModule.RegisterCustomClassLayout(UOpenColorIOConfiguration::StaticClass()->GetFName(), FOnGetDetailCustomizationInstance::CreateStatic(&FOpenColorIOConfigurationCustomization::MakeInstance));
 }
 
@@ -89,6 +90,7 @@ void FOpenColorIOEditorModule::UnregisterCustomizations()
 		FPropertyEditorModule& PropertyModule = FModuleManager::LoadModuleChecked<FPropertyEditorModule>("PropertyEditor");
 		PropertyModule.UnregisterCustomPropertyTypeLayout(UOpenColorIOConfiguration::StaticClass()->GetFName());
 		PropertyModule.UnregisterCustomPropertyTypeLayout(FOpenColorIOColorSpace::StaticStruct()->GetFName());
+		PropertyModule.UnregisterCustomPropertyTypeLayout(FOpenColorIODisplayView::StaticStruct()->GetFName());
 		PropertyModule.UnregisterCustomPropertyTypeLayout(FOpenColorIOColorConversionSettings::StaticStruct()->GetFName());
 	}
 }
diff --git a/Engine/Plugins/Compositing/OpenColorIO/Source/OpenColorIOEditor/Private/Widgets/SOpenColorIOColorSpacePicker.cpp b/Engine/Plugins/Compositing/OpenColorIO/Source/OpenColorIOEditor/Private/Widgets/SOpenColorIOColorSpacePicker.cpp
--- a/Engine/Plugins/Compositing/OpenColorIO/Source/OpenColorIOEditor/Private/Widgets/SOpenColorIOColorSpacePicker.cpp
+++ b/Engine/Plugins/Compositing/OpenColorIO/Source/OpenColorIOEditor/Private/Widgets/SOpenColorIOColorSpacePicker.cpp
@@ -20,6 +20,8 @@ void SOpenColorIOColorSpacePicker::Construct(const FArguments& InArgs)
 	Configuration = InArgs._Config;
 	ColorSpaceSelection = InArgs._InitialColorSpace;
 	RestrictedColorSpace = InArgs._RestrictedColor;
+	DisplayViewSelection = InArgs._InitialDisplayView;
+	bIsDestination = InArgs._IsDestination;
 	OnColorSpaceChanged = InArgs._OnColorSpaceChanged;
 
 	SelectionButton = SNew(SComboButton)
@@ -36,7 +38,11 @@ void SOpenColorIOColorSpacePicker::Construct(const FArguments& InArgs)
 			SNew(STextBlock)
 			.Text(MakeAttributeLambda([=]
 				{
-					if (ColorSpaceSelection.ColorSpaceIndex != INDEX_NONE)
+					if (DisplayViewSelection.IsValid())
+					{
+						return FText::FromString(*DisplayViewSelection.ToString());
+					}
+					else if (ColorSpaceSelection.ColorSpaceIndex != INDEX_NONE)
 					{
 						return FText::FromString(*ColorSpaceSelection.ToString());
 					}
@@ -87,9 +93,25 @@ void SOpenColorIOColorSpacePicker::SetRestrictedColorSpace(const FOpenColorIOCol
 void SOpenColorIOColorSpacePicker::SetCurrentColorSpace(const FOpenColorIOColorSpace& NewColorSpace)
 {
 	ColorSpaceSelection = NewColorSpace;
+	DisplayViewSelection.Reset();
 	
 	// Let listeners know selection has changed
-	OnColorSpaceChanged.ExecuteIfBound(ColorSpaceSelection);
+	OnColorSpaceChanged.ExecuteIfBound(ColorSpaceSelection, DisplayViewSelection);
+
+	//Close our menu
+	if (SelectionButton.IsValid())
+	{
+		SelectionButton->SetIsOpen(false);
+	}
+}
+
+void SOpenColorIOColorSpacePicker::SetCurrentDisplayView(const FOpenColorIODisplayView& NewDisplayView)
+{
+	ColorSpaceSelection.Reset();
+	DisplayViewSelection = NewDisplayView;
+
+	// Let listeners know selection has changed
+	OnColorSpaceChanged.ExecuteIfBound(ColorSpaceSelection, DisplayViewSelection);
 
 	//Close our menu
 	if (SelectionButton.IsValid())
@@ -106,7 +128,7 @@ TSharedRef<SWidget> SOpenColorIOColorSpacePicker::HandleColorSpaceComboButtonMen
 		const bool bShouldCloseWindowAfterClosing = false;
 		FMenuBuilder MenuBuilder(bShouldCloseWindowAfterClosing, nullptr);
 
-		MenuBuilder.BeginSection("AvailableColorSpaces", LOCTEXT("AvailableCoorSpaces", "Available Color Spaces"));
+		MenuBuilder.BeginSection("AvailableColorSpaces", LOCTEXT("AvailableColorSpaces", "Available Color Spaces"));
 		{
 			bool ColorSpaceAdded = false;
 
@@ -149,6 +171,53 @@ TSharedRef<SWidget> SOpenColorIOColorSpacePicker::HandleColorSpaceComboButtonMen
 		}
 		MenuBuilder.EndSection();
 
+		if (bIsDestination)
+		{
+			MenuBuilder.BeginSection("AvailableDisplayViews", LOCTEXT("AvailableDisplayViews", "Available Display-Views"));
+			{
+				bool DisplayViewAdded = false;
+
+				for (int32 i = 0; i < ConfigurationObject->DesiredDisplayViews.Num(); ++i)
+				{
+					const FOpenColorIODisplayView& DisplayView = ConfigurationObject->DesiredDisplayViews[i];
+					if (DisplayView == DisplayViewSelection || !DisplayView.IsValid())
+					{
+						continue;
+					}
+
+					MenuBuilder.AddMenuEntry
+					(
+						FText::FromString(DisplayView.ToString()),
+						FText::FromString(DisplayView.ToString()),
+						FSlateIcon(),
+						FUIAction
+						(
+							FExecuteAction::CreateLambda([this, DisplayView]
+							{
+								SetCurrentDisplayView(DisplayView);
+							}),
+							FCanExecuteAction(),
+							FIsActionChecked::CreateLambda([=]
+							{
+								return DisplayViewSelection == DisplayView;
+							})
+						),
+						NAME_None,
+						EUserInterfaceActionType::RadioButton
+					);
+
+					DisplayViewAdded = true;
+				}
+
+				if (!DisplayViewAdded)
+				{
+					MenuBuilder.AddWidget(SNullWidget::NullWidget, LOCTEXT("NoDisplayViewFound", "No available display-view"), false, false);
+				}
+			}
+			MenuBuilder.EndSection();
+		}
+
+
 		return MenuBuilder.MakeWidget();
 	}
 
@@ -158,12 +227,13 @@ TSharedRef<SWidget> SOpenColorIOColorSpacePicker::HandleColorSpaceComboButtonMen
 FReply SOpenColorIOColorSpacePicker::OnResetToDefault()
 {
 	SetCurrentColorSpace(FOpenColorIOColorSpace());
+	SetCurrentDisplayView(FOpenColorIODisplayView());
 	return FReply::Handled();
 }
 
 EVisibility SOpenColorIOColorSpacePicker::ShouldShowResetToDefaultButton() const
 {
-	return ColorSpaceSelection != FOpenColorIOColorSpace() ? EVisibility::Visible : EVisibility::Hidden;
+	return (ColorSpaceSelection.IsValid() || DisplayViewSelection.IsValid()) ? EVisibility::Visible : EVisibility::Hidden;
 }
 
 #undef LOCTEXT_NAMESPACE
diff --git a/Engine/Plugins/Compositing/OpenColorIO/Source/OpenColorIOEditor/Private/Widgets/SOpenColorIOColorSpacePicker.h b/Engine/Plugins/Compositing/OpenColorIO/Source/OpenColorIOEditor/Private/Widgets/SOpenColorIOColorSpacePicker.h
--- a/Engine/Plugins/Compositing/OpenColorIO/Source/OpenColorIOEditor/Private/Widgets/SOpenColorIOColorSpacePicker.h
+++ b/Engine/Plugins/Compositing/OpenColorIO/Source/OpenColorIOEditor/Private/Widgets/SOpenColorIOColorSpacePicker.h
@@ -16,7 +16,7 @@ class FViewportClient;
 class SComboButton;
 class UOpenColorIOConfiguration;
 
-DECLARE_DELEGATE_OneParam(FOnColorSpaceChanged, const FOpenColorIOColorSpace& /*ColorSpace*/);
+DECLARE_DELEGATE_TwoParams(FOnColorSpaceChanged, const FOpenColorIOColorSpace& /*ColorSpace*/, const FOpenColorIODisplayView& /*DisplayView*/);
 
 class SOpenColorIOColorSpacePicker : public SCompoundWidget
 {
@@ -25,6 +25,8 @@ public:
 		SLATE_ARGUMENT(TWeakObjectPtr<UOpenColorIOConfiguration>, Config)
 		SLATE_ARGUMENT(FOpenColorIOColorSpace, InitialColorSpace)
 		SLATE_ARGUMENT(FOpenColorIOColorSpace, RestrictedColor)
+		SLATE_ARGUMENT(FOpenColorIODisplayView, InitialDisplayView)
+		SLATE_ARGUMENT(bool, IsDestination)
 		SLATE_EVENT(FOnColorSpaceChanged, OnColorSpaceChanged)
 	SLATE_END_ARGS()
 
@@ -41,6 +43,7 @@ protected:
 	
 	/** Called when a selection has been made */
 	void SetCurrentColorSpace(const FOpenColorIOColorSpace& NewColorSpace);
+	void SetCurrentDisplayView(const FOpenColorIODisplayView& NewDisplayView);
 
 	/** Handles color space list menu creation */
 	TSharedRef<SWidget> HandleColorSpaceComboButtonMenuContent();
@@ -56,5 +59,7 @@ protected:
 	TWeakObjectPtr<UOpenColorIOConfiguration> Configuration;
 	FOpenColorIOColorSpace ColorSpaceSelection;
 	FOpenColorIOColorSpace RestrictedColorSpace;
+	FOpenColorIODisplayView DisplayViewSelection;
 	FOnColorSpaceChanged OnColorSpaceChanged;
+	bool bIsDestination = false;
 };
diff --git a/Engine/Plugins/Compositing/OpenColorIO/Source/OpenColorIOEditor/Private/Widgets/SOpenColorIODisplay.cpp b/Engine/Plugins/Compositing/OpenColorIO/Source/OpenColorIOEditor/Private/Widgets/SOpenColorIODisplay.cpp
--- a/Engine/Plugins/Compositing/OpenColorIO/Source/OpenColorIOEditor/Private/Widgets/SOpenColorIODisplay.cpp
+++ b/Engine/Plugins/Compositing/OpenColorIO/Source/OpenColorIOEditor/Private/Widgets/SOpenColorIODisplay.cpp
@@ -66,6 +66,8 @@ void SOpenColorIODisplay::Construct(const FArguments& InArgs)
 		.Config(Configuration.ColorConfiguration.ConfigurationSource)
 		.InitialColorSpace(Configuration.ColorConfiguration.SourceColorSpace)
 		.RestrictedColor(Configuration.ColorConfiguration.DestinationColorSpace)
+		.InitialDisplayView(FOpenColorIODisplayView())
+		.IsDestination(false)
 		.OnColorSpaceChanged(FOnColorSpaceChanged::CreateSP(this, &SOpenColorIODisplay::OnSourceColorSpaceChanged));
 
 	// Source color space picker widget
@@ -74,6 +76,8 @@ void SOpenColorIODisplay::Construct(const FArguments& InArgs)
 		.Config(Configuration.ColorConfiguration.ConfigurationSource)
 		.InitialColorSpace(Configuration.ColorConfiguration.DestinationColorSpace)
 		.RestrictedColor(Configuration.ColorConfiguration.SourceColorSpace)
+		.InitialDisplayView(Configuration.ColorConfiguration.DestinationDisplayView)
+		.IsDestination(true)
 		.OnColorSpaceChanged(FOnColorSpaceChanged::CreateSP(this, &SOpenColorIODisplay::OnDestinationColorSpaceChanged));
 	
 	Section.AddEntry(FToolMenuEntry::InitWidget("SourceColor", SourceColorSpace.ToSharedRef(), FText::GetEmpty(), true, false));
@@ -168,7 +172,7 @@ bool SOpenColorIODisplay::GetDisplayConfigurationState()
 	return Configuration.bIsEnabled;
 }
 
-void SOpenColorIODisplay::OnSourceColorSpaceChanged(const FOpenColorIOColorSpace& NewColorSpace)
+void SOpenColorIODisplay::OnSourceColorSpaceChanged(const FOpenColorIOColorSpace& NewColorSpace, const FOpenColorIODisplayView&)
 {
 	//When color space changes, assign it, update restriction on the other picker and let know our listeners we changed
 	Configuration.ColorConfiguration.SourceColorSpace = NewColorSpace;
@@ -176,10 +180,10 @@ void SOpenColorIODisplay::OnSourceColorSpaceChanged(const FOpenColorIOColorSpace
 	OnDisplayConfigurationChanged.ExecuteIfBound(Configuration);
 }
 
-void SOpenColorIODisplay::OnDestinationColorSpaceChanged(const FOpenColorIOColorSpace& NewColorSpace)
+void SOpenColorIODisplay::OnDestinationColorSpaceChanged(const FOpenColorIOColorSpace& NewColorSpace, const FOpenColorIODisplayView& NewDisplayView)
 {
 	Configuration.ColorConfiguration.DestinationColorSpace = NewColorSpace;
-	SourceColorSpace->SetRestrictedColorSpace(NewColorSpace);
+	Configuration.ColorConfiguration.DestinationDisplayView = NewDisplayView;
 	OnDisplayConfigurationChanged.ExecuteIfBound(Configuration);
 }
 
diff --git a/Engine/Plugins/Compositing/OpenColorIO/Source/OpenColorIOEditor/Private/Widgets/SOpenColorIODisplay.h b/Engine/Plugins/Compositing/OpenColorIO/Source/OpenColorIOEditor/Private/Widgets/SOpenColorIODisplay.h
--- a/Engine/Plugins/Compositing/OpenColorIO/Source/OpenColorIOEditor/Private/Widgets/SOpenColorIODisplay.h
+++ b/Engine/Plugins/Compositing/OpenColorIO/Source/OpenColorIOEditor/Private/Widgets/SOpenColorIODisplay.h
@@ -40,8 +40,8 @@ protected:
 	void ToggleEnableDisplay();
 	bool CanEnableDisplayConfiguration();
 	bool GetDisplayConfigurationState();
-	void OnSourceColorSpaceChanged(const FOpenColorIOColorSpace& NewColorSpace);
-	void OnDestinationColorSpaceChanged(const FOpenColorIOColorSpace& NewColorSpace);
+	void OnSourceColorSpaceChanged(const FOpenColorIOColorSpace& NewColorSpace, const FOpenColorIODisplayView&);
+	void OnDestinationColorSpaceChanged(const FOpenColorIOColorSpace& NewColorSpace, const FOpenColorIODisplayView& NewDisplayView);
 
 protected:
 	FViewport* Viewport;
