OCIO: Exclude display-referred color spaces from the regular color space list now that we have support for selecting displays and views.

#rnx

diff --git a/Engine/Plugins/Compositing/OpenColorIO/Source/OpenColorIOEditor/Private/OpenColorIOColorSpaceCustomization.cpp b/Engine/Plugins/Compositing/OpenColorIO/Source/OpenColorIOEditor/Private/OpenColorIOColorSpaceCustomization.cpp
--- a/Engine/Plugins/Compositing/OpenColorIO/Source/OpenColorIOEditor/Private/OpenColorIOColorSpaceCustomization.cpp
+++ b/Engine/Plugins/Compositing/OpenColorIO/Source/OpenColorIOEditor/Private/OpenColorIOColorSpaceCustomization.cpp
@@ -257,10 +257,10 @@ void FOpenColorIOColorSpaceCustomization::PopulateSubMenu(FMenuBuilder& InMenuBu
 	// To keep track of submenus that were already added
 	TArray<FString> ExistingSubMenus;
 
-	const int32 ColorSpaceCount = CachedConfigFile->getNumColorSpaces();
+	const int32 ColorSpaceCount = CachedConfigFile->getNumColorSpaces(OCIO_NAMESPACE::SEARCH_REFERENCE_SPACE_SCENE, OCIO_NAMESPACE::COLORSPACE_ACTIVE);
 	for (int32 i = 0; i < ColorSpaceCount; ++i)
 	{
-		const char* ColorSpaceName = CachedConfigFile->getColorSpaceNameByIndex(i);
+		const char* ColorSpaceName = CachedConfigFile->getColorSpaceNameByIndex(OCIO_NAMESPACE::SEARCH_REFERENCE_SPACE_SCENE, OCIO_NAMESPACE::COLORSPACE_ACTIVE, i);
 		OCIO_NAMESPACE::ConstColorSpaceRcPtr LibColorSpace = CachedConfigFile->getColorSpace(ColorSpaceName);
 		if (!LibColorSpace)
 		{
@@ -327,10 +327,10 @@ TSharedRef<SWidget> FOpenColorIOColorSpaceCustomization::HandleSourceComboButton
 	{
 		if (bValidConfiguration)
 		{
-			const int32 ColorSpaceCount = CachedConfigFile->getNumColorSpaces();
+			const int32 ColorSpaceCount = CachedConfigFile->getNumColorSpaces(OCIO_NAMESPACE::SEARCH_REFERENCE_SPACE_SCENE, OCIO_NAMESPACE::COLORSPACE_ACTIVE);
 			for (int32 i = 0; i < ColorSpaceCount; ++i)
 			{
-				const char* ColorSpaceName = CachedConfigFile->getColorSpaceNameByIndex(i);
+				const char* ColorSpaceName = CachedConfigFile->getColorSpaceNameByIndex(OCIO_NAMESPACE::SEARCH_REFERENCE_SPACE_SCENE, OCIO_NAMESPACE::COLORSPACE_ACTIVE, i);
 				OCIO_NAMESPACE::ConstColorSpaceRcPtr LibColorSpace = CachedConfigFile->getColorSpace(ColorSpaceName);
 				if (!LibColorSpace)
 				{
