OCIO: Further support for bidirectional display-view transforms.

This improves the module interfaces for 5.1 with a direction enum, while providing us the option of easily enabling bidirectional view transforms further on.

#rnx

diff --git a/Engine/Plugins/Compositing/OpenColorIO/Source/OpenColorIO/Private/OpenColorIOColorSpace.cpp b/Engine/Plugins/Compositing/OpenColorIO/Source/OpenColorIO/Private/OpenColorIOColorSpace.cpp
--- a/Engine/Plugins/Compositing/OpenColorIO/Source/OpenColorIO/Private/OpenColorIOColorSpace.cpp
+++ b/Engine/Plugins/Compositing/OpenColorIO/Source/OpenColorIO/Private/OpenColorIOColorSpace.cpp
@@ -116,7 +116,16 @@ FString FOpenColorIOColorConversionSettings::ToString() const
 	{
 		if (IsDisplayView())
 		{
-			return FString::Printf(TEXT("%s config - %s to %s"), *ConfigurationSource->GetName(), *SourceColorSpace.ToString(), *DestinationDisplayView.ToString());
+			switch (DisplayViewDirection)
+			{
+			case EOpenColorIOViewTransformDirection::Forward:
+				return FString::Printf(TEXT("%s config - %s to %s"), *ConfigurationSource->GetName(), *SourceColorSpace.ToString(), *DestinationDisplayView.ToString());
+			case EOpenColorIOViewTransformDirection::Inverse:
+				return FString::Printf(TEXT("%s config - %s to %s"), *ConfigurationSource->GetName(), *DestinationDisplayView.ToString(), *SourceColorSpace.ToString());
+			default:
+				checkNoEntry();
+				return FString();
+			}
 		}
 		else
 		{
@@ -132,7 +141,16 @@ bool FOpenColorIOColorConversionSettings::IsValid() const
 	{
 		if (IsDisplayView())
 		{
-			return ConfigurationSource->HasTransform(SourceColorSpace.ColorSpaceName, DestinationDisplayView.Display, DestinationDisplayView.View);
+			switch (DisplayViewDirection)
+			{
+			case EOpenColorIOViewTransformDirection::Forward:
+				return ConfigurationSource->HasTransform(SourceColorSpace.ColorSpaceName, DestinationDisplayView.Display, DestinationDisplayView.View, EOpenColorIOViewTransformDirection::Forward);
+			case EOpenColorIOViewTransformDirection::Inverse:
+				return ConfigurationSource->HasTransform(SourceColorSpace.ColorSpaceName, DestinationDisplayView.Display, DestinationDisplayView.View, EOpenColorIOViewTransformDirection::Inverse);
+			default:
+				checkNoEntry();
+				return false;
+			}
 		}
 		else
 		{
@@ -151,21 +169,12 @@ void FOpenColorIOColorConversionSettings::ValidateColorSpaces()
 		{
 			SourceColorSpace.Reset();
 		}
-
-		if (IsDisplayView())
+		if (!ConfigurationSource->HasDesiredColorSpace(DestinationColorSpace))
 		{
-			if (!ConfigurationSource->HasDesiredDisplayView(DestinationDisplayView))
-			{
-				DestinationDisplayView.Reset();
-			}
 			DestinationColorSpace.Reset();
 		}
-		else
+		if (!ConfigurationSource->HasDesiredDisplayView(DestinationDisplayView))
 		{
-			if (!ConfigurationSource->HasDesiredColorSpace(DestinationColorSpace))
-			{
-				DestinationColorSpace.Reset();
-			}
 			DestinationDisplayView.Reset();
 		}
 	}
diff --git a/Engine/Plugins/Compositing/OpenColorIO/Source/OpenColorIO/Private/OpenColorIOColorTransform.cpp b/Engine/Plugins/Compositing/OpenColorIO/Source/OpenColorIO/Private/OpenColorIOColorTransform.cpp
--- a/Engine/Plugins/Compositing/OpenColorIO/Source/OpenColorIO/Private/OpenColorIOColorTransform.cpp
+++ b/Engine/Plugins/Compositing/OpenColorIO/Source/OpenColorIO/Private/OpenColorIOColorTransform.cpp
@@ -41,12 +41,16 @@ namespace {
 
 	OCIO_NAMESPACE::ConstProcessorRcPtr GetTransformProcessor(UOpenColorIOColorTransform* InTransform, const OCIO_NAMESPACE::ConstConfigRcPtr& CurrentConfig)
 	{
-		if (InTransform->bIsDisplayViewType)
+		EOpenColorIOViewTransformDirection DisplayViewDirection;
+
+		if (InTransform->GetDisplayViewDirection(DisplayViewDirection))
 		{
+			OCIO_NAMESPACE::TransformDirection OcioDirection = static_cast<OCIO_NAMESPACE::TransformDirection>(DisplayViewDirection);
+
 			return CurrentConfig->getProcessor(StringCast<ANSICHAR>(*InTransform->SourceColorSpace).Get(),
 				StringCast<ANSICHAR>(*InTransform->Display).Get(),
 				StringCast<ANSICHAR>(*InTransform->View).Get(),
-				OCIO_NAMESPACE::TransformDirection::TRANSFORM_DIR_FORWARD);
+				OcioDirection);
 		}
 		else
 		{
@@ -165,11 +169,11 @@ bool UOpenColorIOColorTransform::Initialize(UOpenColorIOConfiguration* InOwner,
 	return GenerateColorTransformData(InSourceColorSpace, InDestinationColorSpace);
 }
 
-bool UOpenColorIOColorTransform::Initialize(UOpenColorIOConfiguration* InOwner, const FString& InSourceColorSpace, const FString& InDisplay, const FString& InView, bool bInForwardDirection)
+bool UOpenColorIOColorTransform::Initialize(UOpenColorIOConfiguration* InOwner, const FString& InSourceColorSpace, const FString& InDisplay, const FString& InView, EOpenColorIOViewTransformDirection InDirection)
 {
 	check(InOwner);
 	ConfigurationOwner = InOwner;
-	return GenerateColorTransformData(InSourceColorSpace, InDisplay, InView, bInForwardDirection);
+	return GenerateColorTransformData(InSourceColorSpace, InDisplay, InView, InDirection);
 }
 
 void UOpenColorIOColorTransform::Serialize(FArchive& Ar)
@@ -466,16 +470,26 @@ bool UOpenColorIOColorTransform::IsTransform(const FString& InSourceColorSpace,
 	return false;
 }
 
-bool UOpenColorIOColorTransform::IsTransform(const FString& InSourceColorSpace, const FString& InDisplay, const FString& InView, bool bInForwardDirection) const
+bool UOpenColorIOColorTransform::IsTransform(const FString& InSourceColorSpace, const FString& InDisplay, const FString& InView, EOpenColorIOViewTransformDirection InDirection) const
 {
 	if (bIsDisplayViewType)
 	{
-		return SourceColorSpace == InSourceColorSpace && Display == InDisplay && View == InView;
+		return SourceColorSpace == InSourceColorSpace && Display == InDisplay && View == InView && DisplayViewDirection == InDirection;
 	}
 
 	return false;
 }
 
+bool UOpenColorIOColorTransform::GetDisplayViewDirection(EOpenColorIOViewTransformDirection& OutDirection) const
+{
+	if (bIsDisplayViewType)
+	{
+		OutDirection = DisplayViewDirection;
+	}
+
+	return bIsDisplayViewType;
+}
+
 void UOpenColorIOColorTransform::AllColorTransformsCacheResourceShadersForRendering()
 {
 	for (TObjectIterator<UOpenColorIOColorTransform> It; It; ++It)
@@ -506,7 +520,7 @@ bool UOpenColorIOColorTransform::GenerateColorTransformData(const FString& InSou
 	return false;
 }
 
-bool UOpenColorIOColorTransform::GenerateColorTransformData(const FString& InSourceColorSpace, const FString& InDisplay, const FString& InView, bool bInForwardDirection)
+bool UOpenColorIOColorTransform::GenerateColorTransformData(const FString& InSourceColorSpace, const FString& InDisplay, const FString& InView, EOpenColorIOViewTransformDirection InDirection)
 {
 #if WITH_EDITOR && WITH_OCIO
 	if (InSourceColorSpace.IsEmpty() || InDisplay.IsEmpty() || InView.IsEmpty())
@@ -515,9 +529,11 @@ bool UOpenColorIOColorTransform::GenerateColorTransformData(const FString& InSou
 	}
 
 	SourceColorSpace = InSourceColorSpace;
+	DestinationColorSpace.Empty();
 	Display = InDisplay;
 	View = InView;
 	bIsDisplayViewType = true;
+	DisplayViewDirection = InDirection;
 
 	CacheResourceTextures();
 	CacheResourceShadersForRendering(true);
@@ -531,7 +547,16 @@ FString UOpenColorIOColorTransform::GetTransformFriendlyName()
 {
 	if (bIsDisplayViewType)
 	{
-		return SourceColorSpace + TEXT(" to ") + Display + TEXT(" - ") + View;
+		switch (DisplayViewDirection)
+		{
+		case EOpenColorIOViewTransformDirection::Forward:
+			return SourceColorSpace + TEXT(" to ") + Display + TEXT(" - ") + View;
+		case EOpenColorIOViewTransformDirection::Inverse:
+			return Display + TEXT(" - ") + View + TEXT(" to ") + SourceColorSpace;
+		default:
+			checkNoEntry();
+			return FString();
+		}
 	}
 	else
 	{
diff --git a/Engine/Plugins/Compositing/OpenColorIO/Source/OpenColorIO/Private/OpenColorIOConfiguration.cpp b/Engine/Plugins/Compositing/OpenColorIO/Source/OpenColorIO/Private/OpenColorIOConfiguration.cpp
--- a/Engine/Plugins/Compositing/OpenColorIO/Source/OpenColorIO/Private/OpenColorIOConfiguration.cpp
+++ b/Engine/Plugins/Compositing/OpenColorIO/Source/OpenColorIO/Private/OpenColorIOConfiguration.cpp
@@ -59,9 +59,11 @@ bool UOpenColorIOConfiguration::GetRenderResources(ERHIFeatureLevel::Type InFeat
 {
 	TObjectPtr<UOpenColorIOColorTransform>* TransformPtr = ColorTransforms.FindByPredicate([&](const UOpenColorIOColorTransform* InTransform)
 	{
-		if(InTransform->bIsDisplayViewType)
+		EOpenColorIOViewTransformDirection DisplayViewDirection;
+
+		if (InTransform->GetDisplayViewDirection(DisplayViewDirection))
 		{
-			return InTransform->SourceColorSpace == InSettings.SourceColorSpace.ColorSpaceName && InTransform->Display == InSettings.DestinationDisplayView.Display && InTransform->View == InSettings.DestinationDisplayView.View;
+			return InTransform->SourceColorSpace == InSettings.SourceColorSpace.ColorSpaceName && InTransform->Display == InSettings.DestinationDisplayView.Display && InTransform->View == InSettings.DestinationDisplayView.View && InSettings.DisplayViewDirection == DisplayViewDirection;
 		}
 		else
 		{
@@ -71,7 +73,7 @@ bool UOpenColorIOConfiguration::GetRenderResources(ERHIFeatureLevel::Type InFeat
 
 	if (TransformPtr == nullptr)
 	{
-		UE_LOG(LogOpenColorIO, Warning, TEXT("Color transform for [%s] was not found."), *InSettings.SourceColorSpace.ToString());
+		UE_LOG(LogOpenColorIO, Warning, TEXT("Color transform [%s] was not found."), *InSettings.ToString());
 		return false;
 	}
 
@@ -107,11 +109,11 @@ bool UOpenColorIOConfiguration::HasTransform(const FString& InSourceColorSpace,
 	return (TransformData != nullptr);
 }
 
-bool UOpenColorIOConfiguration::HasTransform(const FString& InSourceColorSpace, const FString& InDisplay, const FString& InView)
+bool UOpenColorIOConfiguration::HasTransform(const FString& InSourceColorSpace, const FString& InDisplay, const FString& InView, EOpenColorIOViewTransformDirection InDirection)
 {
 	UE_TRANSITIONAL_OBJECT_PTR(UOpenColorIOColorTransform)* TransformData = ColorTransforms.FindByPredicate([&](const UOpenColorIOColorTransform* InTransformData)
 		{
-			return InTransformData->IsTransform(InSourceColorSpace, InDisplay, InView);
+			return InTransformData->IsTransform(InSourceColorSpace, InDisplay, InView, InDirection);
 		});
 
 	return (TransformData != nullptr);
@@ -119,12 +121,22 @@ bool UOpenColorIOConfiguration::HasTransform(const FString& InSourceColorSpace,
 
 bool UOpenColorIOConfiguration::HasDesiredColorSpace(const FOpenColorIOColorSpace& ColorSpace) const
 {
-	return DesiredColorSpaces.Find(ColorSpace) != INDEX_NONE;
+	if (ColorSpace.IsValid())
+	{
+		return DesiredColorSpaces.Find(ColorSpace) != INDEX_NONE;
+	}
+
+	return false;
 }
 
 bool UOpenColorIOConfiguration::HasDesiredDisplayView(const FOpenColorIODisplayView& DisplayView) const
 {
-	return DesiredDisplayViews.Find(DisplayView) != INDEX_NONE;
+	if (DisplayView.IsValid())
+	{
+		return DesiredDisplayViews.Find(DisplayView) != INDEX_NONE;
+	}
+
+	return false;
 }
 
 bool UOpenColorIOConfiguration::Validate() const
@@ -211,7 +223,7 @@ void UOpenColorIOConfiguration::ReloadExistingColorspaces()
 		for (const FOpenColorIODisplayView& DisplayView : DesiredDisplayViews)
 		{
 			// note: we only support display-view transforms in the forward direction currently.
-			CreateColorTransform(TopColorSpace.ColorSpaceName, DisplayView.Display, DisplayView.View);
+			CreateColorTransform(TopColorSpace.ColorSpaceName, DisplayView.Display, DisplayView.View, EOpenColorIOViewTransformDirection::Forward);
 		}
 	}
 #endif
@@ -330,21 +342,21 @@ void UOpenColorIOConfiguration::CreateColorTransform(const FString& InSourceColo
 	}
 }
 
-void UOpenColorIOConfiguration::CreateColorTransform(const FString& InSourceColorSpace, const FString& InDisplay, const FString& InView, bool bForwardDirection)
+void UOpenColorIOConfiguration::CreateColorTransform(const FString& InSourceColorSpace, const FString& InDisplay, const FString& InView, EOpenColorIOViewTransformDirection InDirection)
 {
 	if (InSourceColorSpace.IsEmpty() || InDisplay.IsEmpty() || InView.IsEmpty())
 	{
 		return;
 	}
 
-	if (HasTransform(InSourceColorSpace, InDisplay, InView))
+	if (HasTransform(InSourceColorSpace, InDisplay, InView, InDirection))
 	{
 		UE_LOG(LogOpenColorIO, Log, TEXT("OCIOConfig already contains %s to %s-%s transform."), *InSourceColorSpace, *InDisplay, *InView);
 		return;
 	}
 
 	UOpenColorIOColorTransform* NewTransform = NewObject<UOpenColorIOColorTransform>(this, NAME_None, RF_NoFlags);
-	const bool bSuccess = NewTransform->Initialize(this, InSourceColorSpace, InDisplay, InView, bForwardDirection);
+	const bool bSuccess = NewTransform->Initialize(this, InSourceColorSpace, InDisplay, InView, InDirection);
 
 	if (bSuccess)
 	{
@@ -462,7 +474,7 @@ void UOpenColorIOConfiguration::PostEditChangeProperty(FPropertyChangedEvent& Pr
 				for (const FOpenColorIODisplayView& DisplayView : DesiredDisplayViews)
 				{
 					// note: we only support display-view transforms in the forward direction currently.
-					CreateColorTransform(TopColorSpace.ColorSpaceName, DisplayView.Display, DisplayView.View);
+					CreateColorTransform(TopColorSpace.ColorSpaceName, DisplayView.Display, DisplayView.View, EOpenColorIOViewTransformDirection::Forward);
 				}
 			}
 		}
@@ -481,7 +493,7 @@ void UOpenColorIOConfiguration::PostEditChangeProperty(FPropertyChangedEvent& Pr
 				for (const FOpenColorIOColorSpace& SourceColorSpace : DesiredColorSpaces)
 				{
 					// note: we only support display-view transforms in the forward direction currently.
-					CreateColorTransform(SourceColorSpace.ColorSpaceName, DisplayView.Display, DisplayView.View);
+					CreateColorTransform(SourceColorSpace.ColorSpaceName, DisplayView.Display, DisplayView.View, EOpenColorIOViewTransformDirection::Forward);
 				}
 			}
 		}
diff --git a/Engine/Plugins/Compositing/OpenColorIO/Source/OpenColorIO/Public/OpenColorIOColorSpace.h b/Engine/Plugins/Compositing/OpenColorIO/Source/OpenColorIO/Public/OpenColorIOColorSpace.h
--- a/Engine/Plugins/Compositing/OpenColorIO/Source/OpenColorIO/Public/OpenColorIOColorSpace.h
+++ b/Engine/Plugins/Compositing/OpenColorIO/Source/OpenColorIO/Public/OpenColorIOColorSpace.h
@@ -69,6 +69,18 @@ public:
 	FString GetFamilyNameAtDepth(int32 InDepth) const;
 };
 
+
+/**
+ * Transformation direction type for display-view transformations.
+ */
+UENUM(BlueprintType)
+enum class EOpenColorIOViewTransformDirection : uint8
+{
+	Forward = 0     UMETA(DisplayName = "Forward"),
+	Inverse = 1     UMETA(DisplayName = "Inverse")
+};
+
+
 USTRUCT(BlueprintType)
 struct OPENCOLORIO_API FOpenColorIODisplayView
 {
@@ -131,6 +143,10 @@ public:
 	UPROPERTY(EditAnywhere, BlueprintReadWrite, Category = ColorSpace)
 	FOpenColorIODisplayView DestinationDisplayView;
 
+	/** The display view direction. */
+	UPROPERTY(EditAnywhere, BlueprintReadWrite, Category = ColorSpace)
+	EOpenColorIOViewTransformDirection DisplayViewDirection = EOpenColorIOViewTransformDirection::Forward;
+
 public:
 
 	/**
diff --git a/Engine/Plugins/Compositing/OpenColorIO/Source/OpenColorIO/Public/OpenColorIOColorTransform.h b/Engine/Plugins/Compositing/OpenColorIO/Source/OpenColorIO/Public/OpenColorIOColorTransform.h
--- a/Engine/Plugins/Compositing/OpenColorIO/Source/OpenColorIO/Public/OpenColorIOColorTransform.h
+++ b/Engine/Plugins/Compositing/OpenColorIO/Source/OpenColorIO/Public/OpenColorIOColorTransform.h
@@ -33,7 +33,7 @@ public:
 
 public:
 	bool Initialize(UOpenColorIOConfiguration* InOwner, const FString& InSourceColorSpace, const FString& InDestinationColorSpace);
-	bool Initialize(UOpenColorIOConfiguration* InOwner, const FString& InSourceColorSpace, const FString& InDisplay, const FString& InView, bool bInForwardDirection = true);
+	bool Initialize(UOpenColorIOConfiguration* InOwner, const FString& InSourceColorSpace, const FString& InDisplay, const FString& InView, EOpenColorIOViewTransformDirection InDirection);
 
 	/**
 	 * Cache resource shaders for cooking on the given shader platform.
@@ -72,7 +72,15 @@ public:
 	bool GetShaderAndLUTResouces(ERHIFeatureLevel::Type InFeatureLevel, FOpenColorIOTransformResource*& OutShaderResource, FTextureResource*& OutLUT3dResource);
 
 	bool IsTransform(const FString& InSourceColorSpace, const FString& InDestinationColorSpace) const;
-	bool IsTransform(const FString& InSourceColorSpace, const FString& InDisplay, const FString& InView, bool bInForwardDirection = true) const;
+	bool IsTransform(const FString& InSourceColorSpace, const FString& InDisplay, const FString& InView, EOpenColorIOViewTransformDirection InDirection) const;
+
+	/**
+	 * Get the display view direction type, when applicable.
+	 *
+	 * @param OutDirection The returned direction.
+	 * @return True if the transformation is of display-view type.
+	 */
+	bool GetDisplayViewDirection(EOpenColorIOViewTransformDirection& OutDirection) const;
 
 	// For all ColorTransforms, UOpenColorIOColorTransform::CacheResourceShadersForRendering
 	static void AllColorTransformsCacheResourceShadersForRendering();
@@ -105,7 +113,7 @@ protected:
 	 * @return true if parameters are valid.
 	 */
 	bool GenerateColorTransformData(const FString& InSourceColorSpace, const FString& InDestinationColorSpace);
-	bool GenerateColorTransformData(const FString& InSourceColorSpace, const FString& InDisplay, const FString& InView, bool bInForwardDirection = true);
+	bool GenerateColorTransformData(const FString& InSourceColorSpace, const FString& InDisplay, const FString& InView, EOpenColorIOViewTransformDirection InDirection);
 
 	/**
 	 * Helper function returning the color space transform name based on source and destination color spaces.
@@ -179,6 +187,9 @@ public:
 
 	UPROPERTY(VisibleAnywhere, Category = "ColorSpace")
 	FString View;
+
+	UPROPERTY(VisibleAnywhere, Category = "ColorSpace")
+	EOpenColorIOViewTransformDirection DisplayViewDirection = EOpenColorIOViewTransformDirection::Forward;
 private:
 
 	/** If the color space requires textures, this will contains the data to do the transform */
diff --git a/Engine/Plugins/Compositing/OpenColorIO/Source/OpenColorIO/Public/OpenColorIOConfiguration.h b/Engine/Plugins/Compositing/OpenColorIO/Source/OpenColorIO/Public/OpenColorIOConfiguration.h
--- a/Engine/Plugins/Compositing/OpenColorIO/Source/OpenColorIO/Public/OpenColorIOConfiguration.h
+++ b/Engine/Plugins/Compositing/OpenColorIO/Source/OpenColorIO/Public/OpenColorIOConfiguration.h
@@ -39,7 +39,7 @@ public:
 	UE_DEPRECATED(5.1, "GetShaderAndLUTResources is deprecated, please use GetRenderResources instead.")
 	bool GetShaderAndLUTResources(ERHIFeatureLevel::Type InFeatureLevel, const FString& InSourceColorSpace, const FString& InDestinationColorSpace, FOpenColorIOTransformResource*& OutShaderResource, FTextureResource*& OutLUT3dResource);
 	bool HasTransform(const FString& InSourceColorSpace, const FString& InDestinationColorSpace);
-	bool HasTransform(const FString& InSourceColorSpace, const FString& InDisplay, const FString& InView);
+	bool HasTransform(const FString& InSourceColorSpace, const FString& InDisplay, const FString& InView, EOpenColorIOViewTransformDirection InDirection);
 	bool HasDesiredColorSpace(const FOpenColorIOColorSpace& ColorSpace) const;
 	bool HasDesiredDisplayView(const FOpenColorIODisplayView& DisplayView) const;
 	bool Validate() const;
@@ -59,7 +59,7 @@ public:
 
 protected:
 	void CreateColorTransform(const FString& InSourceColorSpace, const FString& InDestinationColorSpace);
-	void CreateColorTransform(const FString& InSourceColorSpace, const FString& InDisplay, const FString& InView, bool bForwardDirection = true);
+	void CreateColorTransform(const FString& InSourceColorSpace, const FString& InDisplay, const FString& InView, EOpenColorIOViewTransformDirection InDirection);
 	void CleanupTransforms();
 
 	/** Same as above except user can specify the path manually. */
diff --git a/Engine/Plugins/Compositing/OpenColorIO/Source/OpenColorIOEditor/Private/OpenColorIOColorSpaceConversionCustomization.cpp b/Engine/Plugins/Compositing/OpenColorIO/Source/OpenColorIOEditor/Private/OpenColorIOColorSpaceConversionCustomization.cpp
--- a/Engine/Plugins/Compositing/OpenColorIO/Source/OpenColorIOEditor/Private/OpenColorIOColorSpaceConversionCustomization.cpp
+++ b/Engine/Plugins/Compositing/OpenColorIO/Source/OpenColorIOEditor/Private/OpenColorIOColorSpaceConversionCustomization.cpp
@@ -94,7 +94,11 @@ void FOpenColorIOColorConversionSettingsCustomization::CustomizeChildren(TShared
 	
 			//Create custom rows for source and destination color space of the conversion. Since the struct is hooked to an OCIOConfiguration
 			//We use it to populate the available color spaces instead of using a raw configuration file.
-			if (ChildHandle->GetProperty()->GetFName() == GET_MEMBER_NAME_CHECKED(FOpenColorIOColorConversionSettings, SourceColorSpace))
+			if (ChildHandle->GetProperty()->GetFName() == GET_MEMBER_NAME_CHECKED(FOpenColorIOColorConversionSettings, ConfigurationSource))
+			{
+				StructBuilder.AddProperty(ChildHandle).IsEnabled(true).ShowPropertyButtons(false);
+			}
+			else if (ChildHandle->GetProperty()->GetFName() == GET_MEMBER_NAME_CHECKED(FOpenColorIOColorConversionSettings, SourceColorSpace))
 			{
 				SourceColorSpaceProperty = ChildHandle;
 
@@ -119,10 +123,6 @@ void FOpenColorIOColorConversionSettingsCustomization::CustomizeChildren(TShared
 				ColorSpaceWidget.Visibility(DestinationDisplayViewVisibility);
 				AddPropertyRow(ColorSpaceWidget, ChildHandle, StructCustomizationUtils, true);
 			}
-			else
-			{
-				StructBuilder.AddProperty(ChildHandle).IsEnabled(true).ShowPropertyButtons(false);
-			}
 		}
 	}
 	InStructPropertyHandle->MarkHiddenByCustomization();
