OCIO: Ensure the SVE active functor respects the configuration's enabled state.

#rnx

diff --git a/Engine/Plugins/Compositing/OpenColorIO/Source/OpenColorIO/Private/OpenColorIODisplayExtensionWrapper.cpp b/Engine/Plugins/Compositing/OpenColorIO/Source/OpenColorIO/Private/OpenColorIODisplayExtensionWrapper.cpp
--- a/Engine/Plugins/Compositing/OpenColorIO/Source/OpenColorIO/Private/OpenColorIODisplayExtensionWrapper.cpp
+++ b/Engine/Plugins/Compositing/OpenColorIO/Source/OpenColorIO/Private/OpenColorIODisplayExtensionWrapper.cpp
@@ -21,6 +21,16 @@ void UOpenColorIODisplayExtensionWrapper::CreateDisplayExtensionIfNotExists()
 	check(DisplayExtension.IsValid());
 }
 
+FOpenColorIODisplayConfiguration UOpenColorIODisplayExtensionWrapper::GetOpenColorIOConfiguration() const
+{
+	if (!DisplayExtension.IsValid())
+	{
+		return {};
+	}
+
+	return DisplayExtension->GetDisplayConfiguration();
+}
+
 void UOpenColorIODisplayExtensionWrapper::SetOpenColorIOConfiguration(FOpenColorIODisplayConfiguration InDisplayConfiguration)
 {
 	if (!DisplayExtension.IsValid())
@@ -83,6 +93,12 @@ UOpenColorIODisplayExtensionWrapper* UOpenColorIODisplayExtensionWrapper::Create
 	FSceneViewExtensionIsActiveFunctor IsActiveFunctor;
 	IsActiveFunctor.IsActiveFunction = [](const ISceneViewExtension* SceneViewExtension, const FSceneViewExtensionContext& Context)
 	{
+		// Note: SceneViewExtension this pointer is always passed to functors from FSceneViewExtensionBase::IsActiveThisFrame, we confirm this remains so.
+		check(SceneViewExtension);
+
+		const FOpenColorIODisplayExtension* SVE = static_cast<const FOpenColorIODisplayExtension*>(SceneViewExtension);
+		const bool bIsEnabled = SVE->GetDisplayConfiguration().bIsEnabled;
+
 		if (!Context.Viewport)
 		{
 			return TOptional<bool>();
@@ -94,7 +110,7 @@ UOpenColorIODisplayExtensionWrapper* UOpenColorIODisplayExtensionWrapper::Create
 			// Activate the SVE if it is a PIE viewport.
 			if (Context.Viewport->IsPlayInEditorViewport())
 			{
-				return TOptional<bool>(true);
+				return TOptional<bool>(bIsEnabled);
 			}
 		}
 #endif // WITH_EDITOR
@@ -104,7 +120,7 @@ UOpenColorIODisplayExtensionWrapper* UOpenColorIODisplayExtensionWrapper::Create
 		{
 			if (GameEngine->SceneViewport.IsValid() && (GameEngine->SceneViewport->GetViewport() == Context.Viewport))
 			{
-				return TOptional<bool>(true);
+				return TOptional<bool>(bIsEnabled);
 			}
 		}
 
diff --git a/Engine/Plugins/Compositing/OpenColorIO/Source/OpenColorIO/Public/OpenColorIOColorSpace.h b/Engine/Plugins/Compositing/OpenColorIO/Source/OpenColorIO/Public/OpenColorIOColorSpace.h
--- a/Engine/Plugins/Compositing/OpenColorIO/Source/OpenColorIO/Public/OpenColorIOColorSpace.h
+++ b/Engine/Plugins/Compositing/OpenColorIO/Source/OpenColorIO/Public/OpenColorIOColorSpace.h
@@ -184,7 +184,7 @@ public:
 	 * dictate whether it's applied or not to it
 	 */
 	UPROPERTY(EditAnywhere, BlueprintReadWrite, Category = ColorSpace)
-	bool bIsEnabled = false;
+	bool bIsEnabled = true;
 	
 	/** Conversion to apply when this display is enabled */
 	UPROPERTY(EditAnywhere, BlueprintReadWrite, Category = ColorSpace)
diff --git a/Engine/Plugins/Compositing/OpenColorIO/Source/OpenColorIO/Public/OpenColorIODisplayExtensionWrapper.h b/Engine/Plugins/Compositing/OpenColorIO/Source/OpenColorIO/Public/OpenColorIODisplayExtensionWrapper.h
--- a/Engine/Plugins/Compositing/OpenColorIO/Source/OpenColorIO/Public/OpenColorIODisplayExtensionWrapper.h
+++ b/Engine/Plugins/Compositing/OpenColorIO/Source/OpenColorIO/Public/OpenColorIODisplayExtensionWrapper.h
@@ -18,7 +18,11 @@ class OPENCOLORIO_API UOpenColorIODisplayExtensionWrapper : public UObject
 
 public:
 
-	// Sets a new OCIO configuration.
+	// Gets the display extension OCIO configuration.
+	UFUNCTION(BlueprintCallable, meta = (DisplayName = "Get OpenColorIO Configuration"), Category = "OpenColorIO")
+	FOpenColorIODisplayConfiguration GetOpenColorIOConfiguration() const;
+
+	// Sets the display extension OCIO configuration.
 	UFUNCTION(BlueprintCallable, meta = (DisplayName = "Set OpenColorIO Configuration"), Category = "OpenColorIO")
 	void SetOpenColorIOConfiguration(FOpenColorIODisplayConfiguration InDisplayConfiguration);
 
