OCIO: Cleanup to reuse the same color space picker widget for conversion settings.

#rnx

diff --git a/Engine/Plugins/Compositing/OpenColorIO/Source/OpenColorIOEditor/Private/OpenColorIOColorSpaceConversionCustomization.cpp b/Engine/Plugins/Compositing/OpenColorIO/Source/OpenColorIOEditor/Private/OpenColorIOColorSpaceConversionCustomization.cpp
--- a/Engine/Plugins/Compositing/OpenColorIO/Source/OpenColorIOEditor/Private/OpenColorIOColorSpaceConversionCustomization.cpp
+++ b/Engine/Plugins/Compositing/OpenColorIO/Source/OpenColorIOEditor/Private/OpenColorIOColorSpaceConversionCustomization.cpp
@@ -5,42 +5,30 @@
 #include "Containers/StringConv.h"
 #include "DetailWidgetRow.h"
 #include "Engine/EngineTypes.h"
-#include "Framework/MultiBox/MultiBoxBuilder.h"
 #include "IDetailChildrenBuilder.h"
 #include "IDetailPropertyRow.h"
 #include "IPropertyUtilities.h"
 #include "OpenColorIOColorSpace.h"
 #include "OpenColorIOConfiguration.h"
 #include "PropertyHandle.h"
-#include "SResetToDefaultMenu.h"
-#include "Widgets/Input/SComboButton.h"
-#include "Widgets/Input/SSegmentedControl.h"
+#include "Widgets/SOpenColorIOColorSpacePicker.h"
+#include "Widgets/Text/STextBlock.h"
 
 #define LOCTEXT_NAMESPACE "OpenColorIOColorSpaceConversionCustomization"
 
 void FOpenColorIOColorConversionSettingsCustomization::CustomizeHeader(TSharedRef<IPropertyHandle> InPropertyHandle, FDetailWidgetRow& HeaderRow, IPropertyTypeCustomizationUtils& CustomizationUtils)
 {
-	ColorConversionProperty = InPropertyHandle;
-
-	if (ColorConversionProperty.IsValid())
+	if (InPropertyHandle->GetNumPerObjectValues() == 1 && InPropertyHandle->IsValidHandle())
 	{
-		FProperty* Property = ColorConversionProperty->GetProperty();
-		check(Property && CastField<FStructProperty>(Property) && CastField<FStructProperty>(Property)->Struct && CastField<FStructProperty>(Property)->Struct->IsChildOf(FOpenColorIOColorConversionSettings::StaticStruct()));
-
-		TArray<void*> RawData;
-		ColorConversionProperty->AccessRawData(RawData);
-		if (RawData.Num() > 0)
+		void* StructData = nullptr;
+		if (InPropertyHandle->GetValueData(StructData) == FPropertyAccess::Success)
 		{
-			FOpenColorIOColorConversionSettings* ColorSpaceConversion = reinterpret_cast<FOpenColorIOColorConversionSettings*>(RawData[0]);
-			bIsDestinationDisplayView = ColorSpaceConversion->DestinationDisplayView.IsValid();
-		}
+			ColorSpaceConversion = reinterpret_cast<FOpenColorIOColorConversionSettings*>(StructData);
+			check(ColorSpaceConversion);
 
-		if (ColorConversionProperty->GetNumPerObjectValues() == 1 && ColorConversionProperty->IsValidHandle())
-		{
 			TSharedPtr<IPropertyUtilities> PropertyUtils = CustomizationUtils.GetPropertyUtilities();
 
-			HeaderRow
-				.NameContent()
+			HeaderRow.NameContent()
 				[
 					InPropertyHandle->CreatePropertyNameWidget()
 				]
@@ -54,36 +42,47 @@ void FOpenColorIOColorConversionSettingsCustomization::CustomizeHeader(TSharedRe
 					[
 						SNew(STextBlock)
 						.Text(MakeAttributeLambda([=]
-						{
-							if (ColorConversionProperty.IsValid())
 							{
-								TArray<void*> RawData;
-								ColorConversionProperty->AccessRawData(RawData);
-								if (RawData.Num() > 0)
+								ColorSpaceConversion->ValidateColorSpaces();
+
+								if (ColorSpaceConversion->IsValid())
 								{
-									FOpenColorIOColorConversionSettings* Conversion = reinterpret_cast<FOpenColorIOColorConversionSettings*>(RawData[0]);
-									if (Conversion != nullptr)
-									{
-										Conversion->ValidateColorSpaces();
-										return FText::FromString(*Conversion->ToString());
-									}
+									return FText::FromString(*ColorSpaceConversion->ToString());
 								}
-							}
 
-							return FText::FromString(TEXT("<Invalid Conversion>"));
+								return FText::FromString(TEXT("<Invalid Conversion>"));
 
-						}))
+							}))
 					]
 				].IsEnabled(MakeAttributeLambda([=] { return !InPropertyHandle->IsEditConst() && PropertyUtils->IsPropertyEditingEnabled(); }));
 		}
 	}
 }
 
-void FOpenColorIOColorConversionSettingsCustomization::CustomizeChildren(TSharedRef<IPropertyHandle> InStructPropertyHandle, IDetailChildrenBuilder& StructBuilder, IPropertyTypeCustomizationUtils& StructCustomizationUtils)
+namespace
 {
+	void UpdateColorSettingsStructProperty(const TSharedPtr<IPropertyHandle>& PropertyHandle, const void* Value, const void* Defaults)
+	{
+		if (PropertyHandle.IsValid())
+		{
+			if (FStructProperty* StructProperty = CastField<FStructProperty>(PropertyHandle->GetProperty()))
+			{
+				FString TextValue;
+				StructProperty->Struct->ExportText(TextValue, Value, Defaults, nullptr, EPropertyPortFlags::PPF_None, nullptr);
 
-	TAttribute<EVisibility> DestinationColorSpaceVisibility(this, &FOpenColorIOColorConversionSettingsCustomization::ShouldShowDestinationColorSpace);
-	TAttribute<EVisibility> DestinationDisplayViewVisibility(this, &FOpenColorIOColorConversionSettingsCustomization::ShouldShowDestinationDisplayView);
+				// Note: using set value is preferable for built-in change propagation
+				PropertyHandle->SetValueFromFormattedString(TextValue);
+			}
+		}
+	}
+}
+
+void FOpenColorIOColorConversionSettingsCustomization::CustomizeChildren(TSharedRef<IPropertyHandle> InStructPropertyHandle, IDetailChildrenBuilder& StructBuilder, IPropertyTypeCustomizationUtils& StructCustomizationUtils)
+{
+	if (ColorSpaceConversion == nullptr)
+	{
+		return;
+	}
 
 	uint32 NumberOfChild;
 	if (InStructPropertyHandle->GetNumChildren(NumberOfChild) == FPropertyAccess::Success)
@@ -91,359 +90,107 @@ void FOpenColorIOColorConversionSettingsCustomization::CustomizeChildren(TShared
 		for (uint32 Index = 0; Index < NumberOfChild; ++Index)
 		{
 			TSharedRef<IPropertyHandle> ChildHandle = InStructPropertyHandle->GetChildHandle(Index).ToSharedRef();
-	
-			//Create custom rows for source and destination color space of the conversion. Since the struct is hooked to an OCIOConfiguration
-			//We use it to populate the available color spaces instead of using a raw configuration file.
+
 			if (ChildHandle->GetProperty()->GetFName() == GET_MEMBER_NAME_CHECKED(FOpenColorIOColorConversionSettings, ConfigurationSource))
 			{
 				StructBuilder.AddProperty(ChildHandle).IsEnabled(true).ShowPropertyButtons(false);
+
+				ChildHandle->SetOnPropertyValueChanged(FSimpleDelegate::CreateLambda([this]
+					{
+						TransformSourcePicker->SetConfiguration(ColorSpaceConversion->ConfigurationSource);
+						TransformDestinationPicker->SetConfiguration(ColorSpaceConversion->ConfigurationSource);
+					}));
 			}
 			else if (ChildHandle->GetProperty()->GetFName() == GET_MEMBER_NAME_CHECKED(FOpenColorIOColorConversionSettings, SourceColorSpace))
 			{
 				SourceColorSpaceProperty = ChildHandle;
-
-				FDetailWidgetRow& ColorSpaceWidget = StructBuilder.AddCustomRow(FText::FromName(ChildHandle->GetProperty()->GetFName()));
-				AddPropertyRow(ColorSpaceWidget, ChildHandle, StructCustomizationUtils, false);
-
-				AddDestinationModeRow(StructBuilder, StructCustomizationUtils);
 			}
 			else if (ChildHandle->GetProperty()->GetFName() == GET_MEMBER_NAME_CHECKED(FOpenColorIOColorConversionSettings, DestinationColorSpace))
 			{
 				DestinationColorSpaceProperty = ChildHandle;
-
-				FDetailWidgetRow& ColorSpaceWidget = StructBuilder.AddCustomRow(FText::FromName(ChildHandle->GetProperty()->GetFName()));
-				ColorSpaceWidget.Visibility(DestinationColorSpaceVisibility);
-				AddPropertyRow(ColorSpaceWidget, ChildHandle, StructCustomizationUtils, false);
 			}
 			else if (ChildHandle->GetProperty()->GetFName() == GET_MEMBER_NAME_CHECKED(FOpenColorIOColorConversionSettings, DestinationDisplayView))
 			{
 				DestinationDisplayViewProperty = ChildHandle;
-
-				FDetailWidgetRow& ColorSpaceWidget = StructBuilder.AddCustomRow(FText::FromName(ChildHandle->GetProperty()->GetFName()));
-				ColorSpaceWidget.Visibility(DestinationDisplayViewVisibility);
-				AddPropertyRow(ColorSpaceWidget, ChildHandle, StructCustomizationUtils, true);
 			}
 		}
-	}
-	InStructPropertyHandle->MarkHiddenByCustomization();
-}
 
-void FOpenColorIOColorConversionSettingsCustomization::AddDestinationModeRow(IDetailChildrenBuilder& StructBuilder, IPropertyTypeCustomizationUtils& StructCustomizationUtils)
-{
-	TSharedPtr<IPropertyUtilities> PropertyUtils = StructCustomizationUtils.GetPropertyUtilities();
-	StructBuilder.AddCustomRow(FText::FromString(TEXT("Destination Mode")))
+		// Note: ChildHandle->SetOnPropertyValueChanged isn't automatically called on parent changes, so we explicitely reset the configuration here.
+		InStructPropertyHandle->SetOnPropertyResetToDefault(FSimpleDelegate::CreateSP(this, &FOpenColorIOColorConversionSettingsCustomization::OnConfigurationReset));
+		if (TSharedPtr<IPropertyHandle> ParentHandle = InStructPropertyHandle->GetParentHandle())
+		{
+			ParentHandle->SetOnPropertyResetToDefault(FSimpleDelegate::CreateSP(this, &FOpenColorIOColorConversionSettingsCustomization::OnConfigurationReset));
+		}
+
+		TransformSourcePicker = SNew(SOpenColorIOColorSpacePicker)
+			.Config(ColorSpaceConversion->ConfigurationSource)
+			.InitialColorSpace(ColorSpaceConversion->SourceColorSpace)
+			.RestrictedColor(ColorSpaceConversion->DestinationColorSpace)
+			.InitialDisplayView(FOpenColorIODisplayView())
+			.IsDestination(false)
+			.OnColorSpaceChanged(FOnColorSpaceChanged::CreateLambda([this](const FOpenColorIOColorSpace& NewColorSpace, const FOpenColorIODisplayView& NewDisplayView)
+				{
+					check(SourceColorSpaceProperty.IsValid());
+
+					UpdateColorSettingsStructProperty(SourceColorSpaceProperty, &NewColorSpace, &ColorSpaceConversion->SourceColorSpace);
+					
+					TransformDestinationPicker->SetRestrictedColorSpace(NewColorSpace);
+				}));
+
+		TransformDestinationPicker = SNew(SOpenColorIOColorSpacePicker)
+			.Config(ColorSpaceConversion->ConfigurationSource)
+			.InitialColorSpace(ColorSpaceConversion->DestinationColorSpace)
+			.RestrictedColor(ColorSpaceConversion->SourceColorSpace)
+			.InitialDisplayView(ColorSpaceConversion->DestinationDisplayView)
+			.IsDestination(true)
+			.OnColorSpaceChanged(FOnColorSpaceChanged::CreateLambda([this](const FOpenColorIOColorSpace& NewColorSpace, const FOpenColorIODisplayView& NewDisplayView)
+				{
+					check(DestinationColorSpaceProperty.IsValid());
+					check(DestinationDisplayViewProperty.IsValid());
+
+					UpdateColorSettingsStructProperty(DestinationColorSpaceProperty, &NewColorSpace, &ColorSpaceConversion->DestinationColorSpace);
+					UpdateColorSettingsStructProperty(DestinationDisplayViewProperty, &NewDisplayView, &ColorSpaceConversion->DestinationDisplayView);
+					
+					TransformSourcePicker->SetRestrictedColorSpace(NewColorSpace);
+				}));
+
+		// Source color space picker widget
+		StructBuilder.AddCustomRow(LOCTEXT("TransformSource", "Transform Source"))
 		.NameContent()
 		[
 			SNew(STextBlock)
-			.Text(LOCTEXT("DestinationMode", "Destination Mode"))
-		.ToolTipText(LOCTEXT("DestinationMode_Tooltip", "The destination mode to use, between color space and dislay-view."))
-		.Font(IPropertyTypeCustomizationUtils::GetRegularFont())
+			.Text(LOCTEXT("TransformSource", "Transform Source"))
+			.ToolTipText(LOCTEXT("TransformSource_Tooltip", "The source color space used for the transform."))
+			.Font(StructCustomizationUtils.GetRegularFont())
 		]
 		.ValueContent()
+		.MaxDesiredWidth(512)
 		[
-			SNew(SSegmentedControl<bool>)
-			.Value_Lambda([this]()
-				{
-					return bIsDestinationDisplayView;
-				})
-		.OnValueChanged(this, &FOpenColorIOColorConversionSettingsCustomization::SetDestinationMode)
-			+ SSegmentedControl<bool>::Slot(false)
-			.Text(LOCTEXT("ColorSpace", "Color Space"))
-			.ToolTip(LOCTEXT("ColorSpace_ToolTip",
-				"Select this if you want to use a color space destination."))
-
-			+ SSegmentedControl<bool>::Slot(true)
-			.Text(LOCTEXT("DisplayView", "Display-View"))
-			.ToolTip(LOCTEXT("DisplayView_ToolTip",
-				"Select this if you want to use a display-view destination."))
+			TransformSourcePicker.ToSharedRef()
 		];
-}
-
-void FOpenColorIOColorConversionSettingsCustomization::AddPropertyRow(FDetailWidgetRow& InWidgetRow, TSharedRef<IPropertyHandle> InChildHandle, IPropertyTypeCustomizationUtils& InCustomizationUtils, bool bIsDisplayViewRow) const
-{
-	TSharedPtr<IPropertyUtilities> PropertyUtils = InCustomizationUtils.GetPropertyUtilities();
-	TSharedPtr<SResetToDefaultMenu> ResetToDefaultMenu;
 
-	InWidgetRow
+		StructBuilder.AddCustomRow(LOCTEXT("TransformDestination", "Transform Destination"))
 		.NameContent()
 		[
-			InChildHandle->CreatePropertyNameWidget()
+			SNew(STextBlock)
+			.Text(LOCTEXT("TransformDestination", "Transform Destination"))
+			.ToolTipText(LOCTEXT("TransformDestination_Tooltip", "The destination color space used for the transform."))
+			.Font(StructCustomizationUtils.GetRegularFont())
 		]
 		.ValueContent()
 		.MaxDesiredWidth(512)
 		[
-			SNew(SHorizontalBox)
-			+ SHorizontalBox::Slot()
-			.FillWidth(1.0f)
-			.VAlign(VAlign_Center)
-			[
-				SNew(STextBlock)
-				.Text(MakeAttributeLambda([=]
-				{
-					TArray<void*> RawData;
-					InChildHandle->AccessRawData(RawData);
-					if (RawData.Num() > 0)
-					{
-						if (bIsDisplayViewRow)
-						{
-							FOpenColorIODisplayView* DisplayView = reinterpret_cast<FOpenColorIODisplayView*>(RawData[0]);
-							if (DisplayView != nullptr)
-							{
-								return FText::FromString(*DisplayView->ToString());
-							}
-						}
-						else
-						{
-							FOpenColorIOColorSpace* ColorSpace = reinterpret_cast<FOpenColorIOColorSpace*>(RawData[0]);
-							if (ColorSpace != nullptr)
-							{
-								return FText::FromString(*ColorSpace->ToString());
-							}
-						}
-					}
-				
-					return FText::FromString(TEXT("<Invalid>"));
-				}))
-			]
-			+ SHorizontalBox::Slot()
-			.AutoWidth()
-			.Padding(FMargin(4.0f, 0.0f, 0.0f, 0.0f))
-			.VAlign(VAlign_Center)
-			[
-				SNew(SComboButton)
-				.OnGetMenuContent_Lambda([=]() { return bIsDisplayViewRow ? HandleDisplayViewComboButtonMenuContent(InChildHandle) : HandleColorSpaceComboButtonMenuContent(InChildHandle); })
-				.ContentPadding(FMargin(4.0, 2.0))
-			]
-			+ SHorizontalBox::Slot()
-			.AutoWidth()
-			.Padding(2.0f)
-			[
-				SAssignNew(ResetToDefaultMenu, SResetToDefaultMenu)
-			]
-		].IsEnabled(MakeAttributeLambda([=] { return !InChildHandle->IsEditConst() && PropertyUtils->IsPropertyEditingEnabled(); }));
-		
-		ResetToDefaultMenu->AddProperty(InChildHandle);
-}
-
-TSharedRef<SWidget> FOpenColorIOColorConversionSettingsCustomization::HandleColorSpaceComboButtonMenuContent(TSharedPtr<IPropertyHandle> InPropertyHandle) const
-{
-	FOpenColorIOColorConversionSettings* ColorSpaceConversion = nullptr;
-
-	if (InPropertyHandle.IsValid())
-	{
-		TArray<void*> RawData;
-		ColorConversionProperty->AccessRawData(RawData);
-		if (RawData.Num() > 0)
-		{
-			ColorSpaceConversion = reinterpret_cast<FOpenColorIOColorConversionSettings*>(RawData[0]);
-		}
-	}
-
-	if (ColorSpaceConversion && ColorSpaceConversion->ConfigurationSource)
-	{
-		FOpenColorIOColorSpace RestrictedColorSpace;
-		if (InPropertyHandle == SourceColorSpaceProperty)
-		{
-			RestrictedColorSpace = ColorSpaceConversion->DestinationColorSpace;
-		}
-		else
-		{
-			RestrictedColorSpace = ColorSpaceConversion->SourceColorSpace;
-		}
-
-		// generate menu
-		FMenuBuilder MenuBuilder(true, nullptr);
-
-		MenuBuilder.BeginSection("AvailableColorSpaces", LOCTEXT("AvailableColorSpaces", "Available Color Spaces"));
-		{
-			bool ColorSpaceAdded = false;
-
-			for (int32 i = 0; i < ColorSpaceConversion->ConfigurationSource->DesiredColorSpaces.Num(); ++i)
-			{
-				const FOpenColorIOColorSpace& ColorSpace = ColorSpaceConversion->ConfigurationSource->DesiredColorSpaces[i];
-				if (ColorSpace == RestrictedColorSpace || !ColorSpace.IsValid())
-				{
-					continue;
-				}
-
-				MenuBuilder.AddMenuEntry
-				(
-					FText::FromString(ColorSpace.ToString()),
-					FText::FromString(ColorSpace.ToString()),
-					FSlateIcon(),
-					FUIAction
-					(
-						FExecuteAction::CreateLambda([=] 
-						{
-							if (InPropertyHandle.IsValid())
-							{
-								if (FStructProperty* StructProperty = CastField<FStructProperty>(InPropertyHandle->GetProperty()))
-								{
-									TArray<void*> RawData;
-									InPropertyHandle->AccessRawData(RawData);
-									if (RawData.Num() > 0)
-									{
-										FOpenColorIOColorSpace* PreviousColorSpaceValue = reinterpret_cast<FOpenColorIOColorSpace*>(RawData[0]);
-
-										FString TextValue;
-										StructProperty->Struct->ExportText(TextValue, &ColorSpace, PreviousColorSpaceValue, nullptr, EPropertyPortFlags::PPF_None, nullptr);
-										ensure(InPropertyHandle->SetValueFromFormattedString(TextValue, EPropertyValueSetFlags::DefaultFlags) == FPropertyAccess::Result::Success);
-									}
-								}
-							}
-						}),
-						FCanExecuteAction(),
-						FIsActionChecked::CreateLambda([=] 
-						{
-							if (SourceColorSpaceProperty.IsValid())
-							{
-								TArray<void*> RawData;
-								SourceColorSpaceProperty->AccessRawData(RawData);
-								if (RawData.Num() > 0)
-								{
-									FOpenColorIOColorSpace* PreviousColorSpaceValue = reinterpret_cast<FOpenColorIOColorSpace*>(RawData[0]);
-									return *PreviousColorSpaceValue == ColorSpace;
-								}
-							}
-						
-							return false;
-						})
-					),
-					NAME_None,
-					EUserInterfaceActionType::RadioButton
-				);
-
-				ColorSpaceAdded = true;
-			}
-
-			if (!ColorSpaceAdded)
-			{
-				MenuBuilder.AddWidget(SNullWidget::NullWidget, LOCTEXT("NoColorSpaceFound", "No available color spaces"), false, false);
-			}
-		}
-		MenuBuilder.EndSection();
-
-		return MenuBuilder.MakeWidget();
-	}
-
-	return SNullWidget::NullWidget;
-}
-
-TSharedRef<SWidget> FOpenColorIOColorConversionSettingsCustomization::HandleDisplayViewComboButtonMenuContent(TSharedPtr<IPropertyHandle> InPropertyHandle) const
-{
-	FOpenColorIOColorConversionSettings* ColorSpaceConversion = nullptr;
-
-	if (InPropertyHandle.IsValid())
-	{
-		TArray<void*> RawData;
-		ColorConversionProperty->AccessRawData(RawData);
-		if (RawData.Num() > 0)
-		{
-			ColorSpaceConversion = reinterpret_cast<FOpenColorIOColorConversionSettings*>(RawData[0]);
-		}
-	}
-
-	if (ColorSpaceConversion && ColorSpaceConversion->ConfigurationSource)
-	{
-		// generate menu
-		FMenuBuilder MenuBuilder(true, nullptr);
-
-		MenuBuilder.BeginSection("AvailableDisplayViews", LOCTEXT("AvailableDisplayViews", "Available Display-Views"));
-		{
-			bool DisplayViewAdded = false;
-
-			for (int32 i = 0; i < ColorSpaceConversion->ConfigurationSource->DesiredDisplayViews.Num(); ++i)
-			{
-				const FOpenColorIODisplayView& DisplayView = ColorSpaceConversion->ConfigurationSource->DesiredDisplayViews[i];
-				if (!DisplayView.IsValid())
-				{
-					continue;
-				}
-
-				MenuBuilder.AddMenuEntry
-				(
-					FText::FromString(DisplayView.ToString()),
-					FText::FromString(DisplayView.ToString()),
-					FSlateIcon(),
-					FUIAction
-					(
-						FExecuteAction::CreateLambda([=]
-							{
-								if (InPropertyHandle.IsValid())
-								{
-									if (FStructProperty* StructProperty = CastField<FStructProperty>(InPropertyHandle->GetProperty()))
-									{
-										TArray<void*> RawData;
-										InPropertyHandle->AccessRawData(RawData);
-										if (RawData.Num() > 0)
-										{
-											FOpenColorIODisplayView* PreviousDisplayViewValue = reinterpret_cast<FOpenColorIODisplayView*>(RawData[0]);
-
-											FString TextValue;
-											StructProperty->Struct->ExportText(TextValue, &DisplayView, PreviousDisplayViewValue, nullptr, EPropertyPortFlags::PPF_None, nullptr);
-											ensure(InPropertyHandle->SetValueFromFormattedString(TextValue, EPropertyValueSetFlags::DefaultFlags) == FPropertyAccess::Result::Success);
-										}
-									}
-								}
-							}),
-						FCanExecuteAction(),
-						FIsActionChecked::CreateLambda([=]
-							{
-								return false;
-							})
-						),
-						NAME_None,
-						EUserInterfaceActionType::RadioButton
-					);
-
-				DisplayViewAdded = true;
-			}
-
-			if (!DisplayViewAdded)
-			{
-				MenuBuilder.AddWidget(SNullWidget::NullWidget, LOCTEXT("NoDisplayViewFound", "No available display-view"), false, false);
-			}
-		}
-		MenuBuilder.EndSection();
-
-		return MenuBuilder.MakeWidget();
-	}
-
-	return SNullWidget::NullWidget;
-}
-
-void FOpenColorIOColorConversionSettingsCustomization::SetDestinationMode(bool bInIsDestinationDisplayView)
-{
-	bIsDestinationDisplayView = bInIsDestinationDisplayView;
-
-	if (ColorConversionProperty.IsValid())
-	{
-		TArray<void*> RawData;
-		ColorConversionProperty->AccessRawData(RawData);
-		if (RawData.Num() > 0)
-		{
-			FOpenColorIOColorConversionSettings* ColorSpaceConversion = reinterpret_cast<FOpenColorIOColorConversionSettings*>(RawData[0]);
-
-			if (bIsDestinationDisplayView)
-			{
-				ColorSpaceConversion->DestinationColorSpace.Reset();
-			}
-			else
-			{
-				ColorSpaceConversion->DestinationDisplayView.Reset();
-			}
-		}
+			TransformDestinationPicker.ToSharedRef()
+		];
 	}
-}
 
-EVisibility FOpenColorIOColorConversionSettingsCustomization::ShouldShowDestinationColorSpace() const
-{
-	return bIsDestinationDisplayView ? EVisibility::Hidden : EVisibility::Visible;
+	InStructPropertyHandle->MarkHiddenByCustomization();
 }
 
-EVisibility FOpenColorIOColorConversionSettingsCustomization::ShouldShowDestinationDisplayView() const
+void FOpenColorIOColorConversionSettingsCustomization::OnConfigurationReset()
 {
-	return bIsDestinationDisplayView ? EVisibility::Visible : EVisibility::Hidden;
+	TransformSourcePicker->SetConfiguration(nullptr);
+	TransformDestinationPicker->SetConfiguration(nullptr);
 }
 
 #undef LOCTEXT_NAMESPACE
diff --git a/Engine/Plugins/Compositing/OpenColorIO/Source/OpenColorIOEditor/Private/OpenColorIOColorSpaceConversionCustomization.h b/Engine/Plugins/Compositing/OpenColorIO/Source/OpenColorIOEditor/Private/OpenColorIOColorSpaceConversionCustomization.h
--- a/Engine/Plugins/Compositing/OpenColorIO/Source/OpenColorIOEditor/Private/OpenColorIOColorSpaceConversionCustomization.h
+++ b/Engine/Plugins/Compositing/OpenColorIO/Source/OpenColorIOEditor/Private/OpenColorIOColorSpaceConversionCustomization.h
@@ -6,6 +6,8 @@
 #include "IPropertyTypeCustomization.h"
 #include "Widgets/SWidget.h"
 
+class SOpenColorIOColorSpacePicker;
+
 /**
  * Implements a details view customization for the FOpenColorIOColorConversionSettings
  */
@@ -22,22 +24,9 @@ public:
 	virtual void CustomizeChildren(TSharedRef<class IPropertyHandle> InPropertyHandle, class IDetailChildrenBuilder& StructBuilder, IPropertyTypeCustomizationUtils& PropertyTypeCustomizationUtils) override;
 
 private:
-	void AddDestinationModeRow(IDetailChildrenBuilder& StructBuilder, IPropertyTypeCustomizationUtils& StructCustomizationUtils);
-	void AddPropertyRow(FDetailWidgetRow& InWidgetRow, TSharedRef<IPropertyHandle> InChildHandle, IPropertyTypeCustomizationUtils& InCustomizationUtils, bool bIsDisplayViewRow) const;
-	
-	TSharedRef<SWidget> HandleColorSpaceComboButtonMenuContent(TSharedPtr<IPropertyHandle> InPropertyHandle) const;
-	TSharedRef<SWidget> HandleDisplayViewComboButtonMenuContent(TSharedPtr<IPropertyHandle> InPropertyHandle) const;
-
-	void SetDestinationMode(bool bInIsDestinationDisplayView);
-	/** Controls visibility for widgets when the destination is a color space. */
-	EVisibility ShouldShowDestinationColorSpace() const;
+	/** Callback to reset the configuration in the transform source/destination pickers. */
+	void OnConfigurationReset();
 
-	/** Controls visibility for widgets when the destination is a display-view. */
-	EVisibility ShouldShowDestinationDisplayView() const;
-
-	/** Pointer to the ColorConversion struct property handle. */
-	TSharedPtr<IPropertyHandle> ColorConversionProperty;
-	
 	/** Pointer to the ColorConversion struct member SourceColorSpace property handle. */
 	TSharedPtr<IPropertyHandle> SourceColorSpaceProperty;
 	
@@ -47,5 +36,10 @@ private:
 	/** Pointer to the ColorConversion struct member DestinationColorSpace property handle. */
 	TSharedPtr<IPropertyHandle> DestinationDisplayViewProperty;
 
-	bool bIsDestinationDisplayView = false;
+	/** ColorSpace pickers reference to update them when config asset is changed */
+	TSharedPtr<SOpenColorIOColorSpacePicker> TransformSourcePicker = nullptr;
+	TSharedPtr<SOpenColorIOColorSpacePicker> TransformDestinationPicker = nullptr;
+
+	/** Raw pointer to the conversion settings struct. */
+	struct FOpenColorIOColorConversionSettings* ColorSpaceConversion = nullptr;
 };
diff --git a/Engine/Plugins/Compositing/OpenColorIO/Source/OpenColorIOEditor/Private/Widgets/SOpenColorIODisplay.cpp b/Engine/Plugins/Compositing/OpenColorIO/Source/OpenColorIOEditor/Private/Widgets/SOpenColorIODisplay.cpp
--- a/Engine/Plugins/Compositing/OpenColorIO/Source/OpenColorIOEditor/Private/Widgets/SOpenColorIODisplay.cpp
+++ b/Engine/Plugins/Compositing/OpenColorIO/Source/OpenColorIOEditor/Private/Widgets/SOpenColorIODisplay.cpp
@@ -61,7 +61,7 @@ void SOpenColorIODisplay::Construct(const FArguments& InArgs)
 	));
 
 	// Source color space picker widget
-	SourceColorSpace =
+	TransformSourcePicker =
 		SNew(SOpenColorIOColorSpacePicker)
 		.Config(Configuration.ColorConfiguration.ConfigurationSource)
 		.InitialColorSpace(Configuration.ColorConfiguration.SourceColorSpace)
@@ -71,7 +71,7 @@ void SOpenColorIODisplay::Construct(const FArguments& InArgs)
 		.OnColorSpaceChanged(FOnColorSpaceChanged::CreateSP(this, &SOpenColorIODisplay::OnSourceColorSpaceChanged));
 
 	// Source color space picker widget
-	DestinationColorSpace =
+	TransformDestinationPicker =
 		SNew(SOpenColorIOColorSpacePicker)
 		.Config(Configuration.ColorConfiguration.ConfigurationSource)
 		.InitialColorSpace(Configuration.ColorConfiguration.DestinationColorSpace)
@@ -80,8 +80,8 @@ void SOpenColorIODisplay::Construct(const FArguments& InArgs)
 		.IsDestination(true)
 		.OnColorSpaceChanged(FOnColorSpaceChanged::CreateSP(this, &SOpenColorIODisplay::OnDestinationColorSpaceChanged));
 	
-	Section.AddEntry(FToolMenuEntry::InitWidget("SourceColor", SourceColorSpace.ToSharedRef(), FText::GetEmpty(), true, false));
-	Section.AddEntry(FToolMenuEntry::InitWidget("DestinationColor", DestinationColorSpace.ToSharedRef(), FText::GetEmpty(), true, false));
+	Section.AddEntry(FToolMenuEntry::InitWidget("SourceColor", TransformSourcePicker.ToSharedRef(), FText::GetEmpty(), true, false));
+	Section.AddEntry(FToolMenuEntry::InitWidget("DestinationColor", TransformDestinationPicker.ToSharedRef(), FText::GetEmpty(), true, false));
 
 	TSharedRef<SWidget> Widget = UToolMenus::Get()->GenerateWidget(Menu);
 	ChildSlot
@@ -148,8 +148,8 @@ void SOpenColorIODisplay::OnConfigSelected(const FAssetData& AssetData)
 
 	UOpenColorIOConfiguration* NewConfig = Cast<UOpenColorIOConfiguration>(AssetData.GetAsset());
 	Configuration.ColorConfiguration.ConfigurationSource = NewConfig;
-	SourceColorSpace->SetConfiguration(NewConfig);
-	DestinationColorSpace->SetConfiguration(NewConfig);
+	TransformSourcePicker->SetConfiguration(NewConfig);
+	TransformDestinationPicker->SetConfiguration(NewConfig);
 
 	OnDisplayConfigurationChanged.ExecuteIfBound(Configuration);
 }
@@ -176,13 +176,16 @@ void SOpenColorIODisplay::OnSourceColorSpaceChanged(const FOpenColorIOColorSpace
 {
 	//When color space changes, assign it, update restriction on the other picker and let know our listeners we changed
 	Configuration.ColorConfiguration.SourceColorSpace = NewColorSpace;
-	DestinationColorSpace->SetRestrictedColorSpace(NewColorSpace);
+	TransformDestinationPicker->SetRestrictedColorSpace(NewColorSpace);
+
 	OnDisplayConfigurationChanged.ExecuteIfBound(Configuration);
 }
 
 void SOpenColorIODisplay::OnDestinationColorSpaceChanged(const FOpenColorIOColorSpace& NewColorSpace, const FOpenColorIODisplayView& NewDisplayView)
 {
 	Configuration.ColorConfiguration.DestinationColorSpace = NewColorSpace;
+	TransformSourcePicker->SetRestrictedColorSpace(NewColorSpace);
+
 	Configuration.ColorConfiguration.DestinationDisplayView = NewDisplayView;
 	OnDisplayConfigurationChanged.ExecuteIfBound(Configuration);
 }
diff --git a/Engine/Plugins/Compositing/OpenColorIO/Source/OpenColorIOEditor/Private/Widgets/SOpenColorIODisplay.h b/Engine/Plugins/Compositing/OpenColorIO/Source/OpenColorIOEditor/Private/Widgets/SOpenColorIODisplay.h
--- a/Engine/Plugins/Compositing/OpenColorIO/Source/OpenColorIOEditor/Private/Widgets/SOpenColorIODisplay.h
+++ b/Engine/Plugins/Compositing/OpenColorIO/Source/OpenColorIOEditor/Private/Widgets/SOpenColorIODisplay.h
@@ -50,8 +50,8 @@ protected:
 	TSharedPtr<SWidget> ClassPicker;
 
 	/** ColorSpace pickers reference to update them when config asset is changed */
-	TSharedPtr<SOpenColorIOColorSpacePicker> SourceColorSpace;
-	TSharedPtr<SOpenColorIOColorSpacePicker> DestinationColorSpace;
+	TSharedPtr<SOpenColorIOColorSpacePicker> TransformSourcePicker;
+	TSharedPtr<SOpenColorIOColorSpacePicker> TransformDestinationPicker;
 
 	/** Current configuration */
 	FOpenColorIODisplayConfiguration Configuration;
