OpenColorIO shader compiling WITH_EDITOR cleanup.
None of the shader compiling code runs in non-Editor mode, so don't compile it for those targets.

diff --git a/Engine/Plugins/Compositing/OpenColorIO/Source/OpenColorIO/Private/OpenColorIOShaderCompilationManager.cpp b/Engine/Plugins/Compositing/OpenColorIO/Source/OpenColorIO/Private/OpenColorIOShaderCompilationManager.cpp
--- a/Engine/Plugins/Compositing/OpenColorIO/Source/OpenColorIO/Private/OpenColorIOShaderCompilationManager.cpp
+++ b/Engine/Plugins/Compositing/OpenColorIO/Source/OpenColorIO/Private/OpenColorIOShaderCompilationManager.cpp
@@ -2,16 +2,15 @@
 
 #include "OpenColorIOShaderCompilationManager.h"
 
+#if WITH_EDITOR
+
 #include "HAL/FileManager.h"
 #include "Misc/Paths.h"
 #include "OpenColorIOShared.h"
 #include "ShaderCompiler.h"
 
-#if WITH_EDITOR
 #include "Interfaces/IShaderFormat.h"
 #include "Interfaces/ITargetPlatformManagerModule.h"
-#endif
-
 
 DEFINE_LOG_CATEGORY_STATIC(LogOpenColorIOShaderCompiler, All, All);
 
@@ -22,15 +21,11 @@ static FAutoConsoleVariableRef CVarShowOpenColorIOShaderWarnings(
 	TEXT("When set to 1, will display all warnings from OpenColorIO shader compiles.")
 	);
 
-
-
-OPENCOLORIO_API FOpenColorIOShaderCompilationManager GOpenColorIOShaderCompilationManager;
+FOpenColorIOShaderCompilationManager GOpenColorIOShaderCompilationManager;
 
 void FOpenColorIOShaderCompilationManager::Tick(float DeltaSeconds)
 {
-#if WITH_EDITOR
 	ProcessAsyncResults();
-#endif
 }
 
 FOpenColorIOShaderCompilationManager::FOpenColorIOShaderCompilationManager()
@@ -63,9 +58,8 @@ void FOpenColorIOShaderCompilationManager::InitWorkerInfo()
 	}	
 }
 
-OPENCOLORIO_API void FOpenColorIOShaderCompilationManager::AddJobs(TArray<FShaderCommonCompileJobPtr> InNewJobs)
+void FOpenColorIOShaderCompilationManager::AddJobs(TArray<FShaderCommonCompileJobPtr> InNewJobs)
 {
-#if WITH_EDITOR
 	check(IsInGameThread());
 	JobQueue.Append(InNewJobs);
 
@@ -92,13 +86,10 @@ OPENCOLORIO_API void FOpenColorIOShaderCompilationManager::AddJobs(TArray<FShade
 		}
 	}
 	GShaderCompilingManager->SubmitJobs(InNewJobs, FString(), FString());
-#endif
 }
 
-
 void FOpenColorIOShaderCompilationManager::ProcessAsyncResults()
 {
-#if WITH_EDITOR
 	check(IsInGameThread());
 
 	// Process the results from the shader compile worker
@@ -144,15 +135,12 @@ void FOpenColorIOShaderCompilationManager::ProcessAsyncResults()
 	{
 		ProcessCompiledOpenColorIOShaderMaps(PendingFinalizeOpenColorIOShaderMaps, 10);
 	}
-#endif
 }
 
-
 void FOpenColorIOShaderCompilationManager::ProcessCompiledOpenColorIOShaderMaps(
 	TMap<int32, FOpenColorIOShaderMapFinalizeResults>& CompiledShaderMaps,
 	float TimeBudget)
 {
-#if WITH_EDITOR
 	// Keeps shader maps alive as they are passed from the shader compiler and applied to the owning ColorTransform
 	TArray<TRefCountPtr<FOpenColorIOShaderMap> > LocalShaderMapReferences;
 	TMap<FOpenColorIOTransformResource*, FOpenColorIOShaderMap*> TransformsToUpdate;
@@ -323,16 +311,13 @@ void FOpenColorIOShaderCompilationManager::ProcessCompiledOpenColorIOShaderMaps(
 				});
 		}
 	}
-#endif
 }
 
-
 void FOpenColorIOShaderCompilationManager::FinishCompilation(const TCHAR* InTransformName, const TArray<int32>& ShaderMapIdsToFinishCompiling)
 {
-#if WITH_EDITOR
 	check(!FPlatformProperties::RequiresCookedData());
 	GShaderCompilingManager->FinishCompilation(NULL, ShaderMapIdsToFinishCompiling);
 	ProcessAsyncResults();	// grab compiled shader maps and assign them to their resources
-#endif
 }
 
+#endif // WITH_EDITOR
diff --git a/Engine/Plugins/Compositing/OpenColorIO/Source/OpenColorIO/Private/OpenColorIOShaderMap.cpp b/Engine/Plugins/Compositing/OpenColorIO/Source/OpenColorIO/Private/OpenColorIOShaderMap.cpp
--- a/Engine/Plugins/Compositing/OpenColorIO/Source/OpenColorIO/Private/OpenColorIOShaderMap.cpp
+++ b/Engine/Plugins/Compositing/OpenColorIO/Source/OpenColorIO/Private/OpenColorIOShaderMap.cpp
@@ -48,6 +48,7 @@ namespace OpenColorIOShaderCookStats
 TMap<FOpenColorIOShaderMapId, FOpenColorIOShaderMap*> FOpenColorIOShaderMap::GIdToOpenColorIOShaderMap[SP_NumPlatforms];
 TArray<FOpenColorIOShaderMap*> FOpenColorIOShaderMap::AllOpenColorIOShaderMaps;
 
+#if WITH_EDITOR
 // The Id of 0 is reserved for global shaders
 uint32 FOpenColorIOShaderMap::NextCompilingId = 2;
 
@@ -57,16 +58,13 @@ uint32 FOpenColorIOShaderMap::NextCompilingId = 2;
  * Uses a TRefCountPtr as this will be the only reference to a shader map while it is being compiled.
  */
 TMap<TRefCountPtr<FOpenColorIOShaderMap>, TArray<FOpenColorIOTransformResource*> > FOpenColorIOShaderMap::OpenColorIOShaderMapsBeingCompiled;
-
-
+#endif // WITH_EDITOR
 
 static inline bool ShouldCacheOpenColorIOShader(const FOpenColorIOShaderType* InShaderType, EShaderPlatform InPlatform, const FOpenColorIOTransformResource* InColorTransformShader)
 {
 	return InShaderType->ShouldCache(InPlatform, InColorTransformShader) && InColorTransformShader->ShouldCache(InPlatform, InShaderType);
 }
 
-
-
 /** Called for every color transform shader to update the appropriate stats. */
 void UpdateOpenColorIOShaderCompilingStats(const FOpenColorIOTransformResource* InShader)
 {
@@ -131,9 +129,10 @@ bool FOpenColorIOShaderMapId::operator==(const FOpenColorIOShaderMapId& InRefere
 	return true;
 }
 
+#if WITH_EDITOR
+
 void FOpenColorIOShaderMapId::AppendKeyString(FString& OutKeyString) const
 {
-#if WITH_EDITOR
 	OutKeyString.Append(ShaderCodeAndConfigHash);
 	OutKeyString.AppendChar('_');
 
@@ -146,14 +145,8 @@ void FOpenColorIOShaderMapId::AppendKeyString(FString& OutKeyString) const
 
 	// Add the inputs for any shaders that are stored inline in the shader map
 	AppendKeyStringShaderDependencies(MakeArrayView(ShaderTypeDependencies), LayoutParams, OutKeyString);
-
-#endif //WITH_EDITOR
 }
 
-/**
- * Enqueues a compilation for a new shader of this type.
- * @param InColorTransform - The ColorTransform to link the shader with.
- */
 void FOpenColorIOShaderType::BeginCompileShader(
 	uint32 InShaderMapId,
 	const FOpenColorIOTransformResource* InColorTransform,
@@ -201,11 +194,6 @@ void FOpenColorIOShaderType::BeginCompileShader(
 	OutNewJobs.Add(FShaderCommonCompileJobPtr(NewJob));
 }
 
-/**
- * Either creates a new instance of this type or returns an equivalent existing shader.
- * @param InShaderMapHash - Precomputed hash of the shader map 
- * @param InCurrentJob - Compile job that was enqueued by BeginCompileShader.
- */
 FShader* FOpenColorIOShaderType::FinishCompileShader(
 	const FSHAHash& InShaderMapHash,
 	const FShaderCompileJob& InCurrentJob,
@@ -221,6 +209,8 @@ FShader* FOpenColorIOShaderType::FinishCompileShader(
 	return Shader;
 }
 
+#endif // WITH_EDITOR
+
 /**
  * Finds the shader map for a color transform.
  * @param InShaderMapId - The color transform id and static parameter set identifying the shader map
@@ -233,34 +223,31 @@ FOpenColorIOShaderMap* FOpenColorIOShaderMap::FindId(const FOpenColorIOShaderMap
 	return GIdToOpenColorIOShaderMap[InPlatform].FindRef(InShaderMapId);
 }
 
+#if WITH_EDITOR
+
 void OpenColorIOShaderMapAppendKeyString(EShaderPlatform InPlatform, FString& OutKeyString)
 {
-#if WITH_EDITOR && WITH_OCIO
+#if WITH_OCIO
 	//Keep library version in the DDC key to invalidate it once we move to a new library
 	OutKeyString += TEXT("OCIOVersion");
 	OutKeyString += TEXT(OCIO_VERSION);
 	OutKeyString += TEXT("_");
-#endif //WITH_EDITOR && WITH_OCIO
+#endif // WITH_OCIO
 }
 
 /** Creates a string key for the derived data cache given a shader map id. */
 static FString GetOpenColorIOShaderMapKeyString(const FOpenColorIOShaderMapId& InShaderMapId, EShaderPlatform InPlatform)
 {
-#if WITH_EDITOR
 	const FName Format = LegacyShaderPlatformToShaderFormat(InPlatform);
 	FString ShaderMapKeyString = Format.ToString() + TEXT("_") + FString(FString::FromInt(GetTargetPlatformManagerRef().ShaderFormatVersion(Format))) + TEXT("_");
 	OpenColorIOShaderMapAppendKeyString(InPlatform, ShaderMapKeyString);
 	ShaderMapAppendKeyString(InPlatform, ShaderMapKeyString);
 	InShaderMapId.AppendKeyString(ShaderMapKeyString);
 	return FDerivedDataCacheInterface::BuildCacheKey(TEXT("OCIOSM"), OPENCOLORIO_DERIVEDDATA_VER, *ShaderMapKeyString);
-#else
-	return FString();
-#endif
 }
 
 void FOpenColorIOShaderMap::LoadFromDerivedDataCache(const FOpenColorIOTransformResource* InColorTransform, const FOpenColorIOShaderMapId& InShaderMapId, EShaderPlatform InPlatform, TRefCountPtr<FOpenColorIOShaderMap>& InOutShaderMap)
 {
-#if WITH_EDITOR
 	if (InOutShaderMap != nullptr)
 	{
 		check(InOutShaderMap->GetShaderPlatform() == InPlatform);
@@ -300,12 +287,10 @@ void FOpenColorIOShaderMap::LoadFromDerivedDataCache(const FOpenColorIOTransform
 		}
 		INC_FLOAT_STAT_BY(STAT_ShaderCompiling_DDCLoading,(float)OpenColorIOShaderDDCTime);
 	}
-#endif
 }
 
 void FOpenColorIOShaderMap::SaveToDerivedDataCache()
 {
-#if WITH_EDITOR
 	COOK_STAT(auto Timer = OpenColorIOShaderCookStats::UsageStats.TimeSyncWork());
 	TArray<uint8> SaveData;
 	FMemoryWriter Ar(SaveData, true);
@@ -313,7 +298,6 @@ void FOpenColorIOShaderMap::SaveToDerivedDataCache()
 
 	GetDerivedDataCacheRef().Put(*GetOpenColorIOShaderMapKeyString(GetContent()->ShaderMapId, GetShaderPlatform()), SaveData, FStringView(*GetFriendlyName()));
 	COOK_STAT(Timer.AddMiss(SaveData.Num()));
-#endif
 }
 
 /**
@@ -506,6 +490,8 @@ bool FOpenColorIOShaderMap::TryToAddToExistingCompilationTask(FOpenColorIOTransf
 	return false;
 }
 
+#endif // WITH_EDITOR
+
 bool FOpenColorIOShaderMap::IsOpenColorIOShaderComplete(const FOpenColorIOTransformResource* InColorTransform, const FOpenColorIOShaderType* InShaderType, bool bSilent)
 {
 	// If we should cache this color transform, it's incomplete if the shader is missing
@@ -526,13 +512,14 @@ bool FOpenColorIOShaderMap::IsComplete(const FOpenColorIOTransformResource* InCo
 	// Make sure we are operating on a referenced shader map or the below Find will cause this shader map to be deleted,
 	// Since it creates a temporary ref counted pointer.
 	check(NumRefs > 0);
-	const TArray<FOpenColorIOTransformResource*>* CorrespondingColorTransforms = FOpenColorIOShaderMap::OpenColorIOShaderMapsBeingCompiled.Find(this);
 
-	if (CorrespondingColorTransforms)
+#if WITH_EDITOR
+	if (const TArray<FOpenColorIOTransformResource*>* CorrespondingColorTransforms = FOpenColorIOShaderMap::OpenColorIOShaderMapsBeingCompiled.Find(this))
 	{
 		check(!bCompilationFinalized);
 		return false;
 	}
+#endif // WITH_EDITOR
 
 	// Iterate over all shader types.
 	for(TLinkedList<FShaderType*>::TIterator ShaderTypeIt(FShaderType::GetTypeList());ShaderTypeIt;ShaderTypeIt.Next())
@@ -613,19 +600,6 @@ FOpenColorIOShaderMap::~FOpenColorIOShaderMap()
 	AllOpenColorIOShaderMaps.RemoveSwap(this);
 }
 
-/**
- * Removes all entries in the cache with exceptions based on a shader type
- * @param ShaderType - The shader type to flush
- */
-void FOpenColorIOShaderMap::FlushShadersByShaderType(const FShaderType* InShaderType)
-{
-	if (InShaderType->GetOpenColorIOShaderType())
-	{
-		GetMutableContent()->RemoveShaderTypePermutaion(InShaderType->GetOpenColorIOShaderType(), /* PermutationId = */ 0);	
-	}
-}
-
-
 bool FOpenColorIOShaderMap::Serialize(FArchive& Ar, bool bInlineShaderResources)
 {
 	// Note: This is saved to the DDC, not into packages (except when cooked)
@@ -634,6 +608,7 @@ bool FOpenColorIOShaderMap::Serialize(FArchive& Ar, bool bInlineShaderResources)
 	return Super::Serialize(Ar, bInlineShaderResources, false);
 }
 
+#if WITH_EDITOR
 void FOpenColorIOShaderMap::RemovePendingColorTransform(FOpenColorIOTransformResource* InColorTransform)
 {
 	for (TMap<TRefCountPtr<FOpenColorIOShaderMap>, TArray<FOpenColorIOTransformResource*> >::TIterator It(OpenColorIOShaderMapsBeingCompiled); It; ++It)
@@ -652,4 +627,4 @@ void FOpenColorIOShaderMap::RemovePendingColorTransform(FOpenColorIOTransformRes
 #endif
 	}
 }
-
+#endif // WITH_EDITOR
diff --git a/Engine/Plugins/Compositing/OpenColorIO/Source/OpenColorIO/Private/OpenColorIOShared.cpp b/Engine/Plugins/Compositing/OpenColorIO/Source/OpenColorIO/Private/OpenColorIOShared.cpp
--- a/Engine/Plugins/Compositing/OpenColorIO/Source/OpenColorIO/Private/OpenColorIOShared.cpp
+++ b/Engine/Plugins/Compositing/OpenColorIO/Source/OpenColorIO/Private/OpenColorIOShared.cpp
@@ -28,13 +28,17 @@ IMPLEMENT_TYPE_LAYOUT(FOpenColorIOShaderMapContent);
 
 FOpenColorIOTransformResource::~FOpenColorIOTransformResource()
 {
+#if WITH_EDITOR
 	FOpenColorIOShaderMap::RemovePendingColorTransform(this);
+#endif
 }
 
+#if WITH_EDITOR
 /** Populates OutEnvironment with defines needed to compile shaders for this color transform. */
 void FOpenColorIOTransformResource::SetupShaderCompilationEnvironment(EShaderPlatform InPlatform, FShaderCompilerEnvironment& OutEnvironment) const
 {
 }
+#endif // WITH_EDITOR
 
 OPENCOLORIO_API bool FOpenColorIOTransformResource::ShouldCache(EShaderPlatform InPlatform, const FShaderType* InShaderType) const
 {
@@ -42,9 +46,10 @@ OPENCOLORIO_API bool FOpenColorIOTransformResource::ShouldCache(EShaderPlatform
 	return true;
 }
 
+#if WITH_EDITOR
+
 OPENCOLORIO_API void FOpenColorIOTransformResource::CancelCompilation()
 {
-#if WITH_EDITOR
 	if (IsInGameThread())
 	{
 		FOpenColorIOShaderMap::RemovePendingColorTransform(this);
@@ -52,7 +57,6 @@ OPENCOLORIO_API void FOpenColorIOTransformResource::CancelCompilation()
 		UE_LOG(LogShaders, Log, TEXT("CancelCompilation %p."), this);
 		OutstandingCompileShaderMapIds.Empty();
 	}
-#endif
 }
 
 OPENCOLORIO_API void FOpenColorIOTransformResource::RemoveOutstandingCompileId(const int32 InOldOutstandingCompileShaderMapId)
@@ -63,9 +67,13 @@ OPENCOLORIO_API void FOpenColorIOTransformResource::RemoveOutstandingCompileId(c
 	}
 }
 
+#endif // WITH_EDITOR
+
 OPENCOLORIO_API void FOpenColorIOTransformResource::Invalidate()
 {
+#if WITH_EDITOR
 	CancelCompilation();
+#endif
 	ReleaseShaderMap();
 }
 
@@ -103,8 +111,8 @@ OPENCOLORIO_API void FOpenColorIOTransformResource::GetShaderMapId(EShaderPlatfo
 		OutId.FeatureLevel = GetFeatureLevel();
 
 		OutId.ShaderCodeAndConfigHash = ShaderCodeAndConfigHash;
-		OutId.SetShaderDependencies(ShaderTypes, InPlatform);
 #if WITH_EDITOR
+		OutId.SetShaderDependencies(ShaderTypes, InPlatform);
 		OutId.LayoutParams.InitializeForPlatform(TargetPlatform);
 #else
 		if (TargetPlatform != nullptr)
@@ -223,11 +231,13 @@ OPENCOLORIO_API  void FOpenColorIOTransformResource::SetRenderingThreadShaderMap
 
 OPENCOLORIO_API  bool FOpenColorIOTransformResource::IsCompilationFinished() const
 {
-	bool bRet = GameThreadShaderMap && GameThreadShaderMap.IsValid() && GameThreadShaderMap->IsCompilationFinalized();
+#if WITH_EDITOR
 	if (OutstandingCompileShaderMapIds.Num() == 0)
 	{
 		return true;
 	}
+#endif
+	bool bRet = GameThreadShaderMap && GameThreadShaderMap.IsValid() && GameThreadShaderMap->IsCompilationFinalized();
 	return bRet;
 }
 
@@ -273,6 +283,7 @@ bool FOpenColorIOTransformResource::CacheShaders(const FOpenColorIOShaderMapId&
 			// Attempt to load from the derived data cache if we are uncooked
 			if ((!GameThreadShaderMap || !GameThreadShaderMap->IsComplete(this, true)) && !FPlatformProperties::RequiresCookedData())
 			{
+#if WITH_EDITOR
 				FOpenColorIOShaderMap::LoadFromDerivedDataCache(this, InShaderMapId, InPlatform, GameThreadShaderMap);
 				if (GameThreadShaderMap && GameThreadShaderMap->IsValid())
 				{
@@ -282,6 +293,7 @@ bool FOpenColorIOTransformResource::CacheShaders(const FOpenColorIOShaderMapId&
 				{
 					UE_LOG(LogTemp, Display, TEXT("Loading shader for OCIO ColorSpace %s from DDC failed. Shader needs recompile."), *GetFriendlyName());
 				}
+#endif // WITH_EDITOR
 			}
 		}
 	}
@@ -297,7 +309,6 @@ bool FOpenColorIOTransformResource::CacheShaders(const FOpenColorIOShaderMapId&
 	{
 		GameThreadShaderMap->AssociateWithAsset(AssetPath);
 	}
-#endif // WITH_EDITOR
 
 	if (GameThreadShaderMap && GameThreadShaderMap->TryToAddToExistingCompilationTask(this))
 	{
@@ -310,13 +321,16 @@ bool FOpenColorIOTransformResource::CacheShaders(const FOpenColorIOShaderMapId&
 		GameThreadShaderMap = nullptr;
 		bSucceeded = true;
 	}
-	else if (!GameThreadShaderMap || !(bAssumeShaderMapIsComplete || GameThreadShaderMap->IsComplete(this, false)))
+	else
+#endif // WITH_EDITOR
+	if (!GameThreadShaderMap || !(bAssumeShaderMapIsComplete || GameThreadShaderMap->IsComplete(this, false)))
 	{
 		if (bContainsInlineShaders || FPlatformProperties::RequiresCookedData())
 		{
 			UE_LOG(LogShaders, Log, TEXT("Can't compile %s with cooked content!"), *GetFriendlyName());
 			GameThreadShaderMap = nullptr;
 		}
+#if WITH_EDITOR
 		else
 		{
 			UE_LOG(LogShaders, Log, TEXT("%s cached shader map for color transform %s, compiling."), GameThreadShaderMap? TEXT("Incomplete") : TEXT("Missing"), *GetFriendlyName());
@@ -330,6 +344,7 @@ bool FOpenColorIOTransformResource::CacheShaders(const FOpenColorIOShaderMapId&
 				GameThreadShaderMap = nullptr;
 			}
 		}
+#endif // WITH_EDITOR
 	}
 	else
 	{
@@ -347,9 +362,9 @@ bool FOpenColorIOTransformResource::CacheShaders(const FOpenColorIOShaderMapId&
 	return bSucceeded;
 }
 
+#if WITH_EDITOR
 void FOpenColorIOTransformResource::FinishCompilation()
 {
-#if WITH_EDITOR
 	TArray<int32> ShaderMapIdsToFinish;
 	GetShaderMapIDsWithUnfinishedCompilation(ShaderMapIdsToFinish);
 
@@ -368,10 +383,8 @@ void FOpenColorIOTransformResource::FinishCompilation()
 		GetShaderMapIDsWithUnfinishedCompilation(ShaderMapIdsToFinish2);
 		ensure(ShaderMapIdsToFinish2.Num() == 0);
 	}
-#endif
 }
 
-
 void FOpenColorIOTransformResource::GetShaderMapIDsWithUnfinishedCompilation(TArray<int32>& OutShaderMapIds)
 {
 	// Build an array of the shader map Id's are not finished compiling.
@@ -395,7 +408,6 @@ void FOpenColorIOTransformResource::GetShaderMapIDsWithUnfinishedCompilation(TAr
  */
 bool FOpenColorIOTransformResource::BeginCompileShaderMap(const FOpenColorIOShaderMapId& InShaderMapId, EShaderPlatform InPlatform, TRefCountPtr<FOpenColorIOShaderMap>& OutShaderMap, bool bApplyCompletedShaderMapForRendering, bool bSynchronous)
 {
-#if WITH_EDITORONLY_DATA
 	bool bSuccess = false;
 
 	STAT(double OpenColorIOCompileTime = 0);
@@ -403,9 +415,7 @@ bool FOpenColorIOTransformResource::BeginCompileShaderMap(const FOpenColorIOShad
 	SCOPE_SECONDS_COUNTER(OpenColorIOCompileTime);
 
 	TRefCountPtr<FOpenColorIOShaderMap> NewShaderMap = new FOpenColorIOShaderMap();
-#if WITH_EDITOR
 	NewShaderMap->AssociateWithAsset(AssetPath);
-#endif
 
 	// Create a shader compiler environment for the material that will be shared by all jobs from this material
 	TRefCountPtr<FSharedShaderCompilerEnvironment> MaterialEnvironment = new FSharedShaderCompilerEnvironment();
@@ -431,15 +441,10 @@ bool FOpenColorIOTransformResource::BeginCompileShaderMap(const FOpenColorIOShad
 	INC_FLOAT_STAT_BY(STAT_ShaderCompiling_OpenColorIOShaders, (float)OpenColorIOCompileTime);
 
 	return true;
-#else
-	UE_LOG(LogShaders, Fatal, TEXT("Compiling of shaders in a build without editordata is not supported."));
-	return false;
-#endif
 }
 
 void FOpenColorIOShaderMapId::SetShaderDependencies(const TArray<FShaderType*>& InShaderTypes, EShaderPlatform InShaderPlatform)
 {
-#if WITH_EDITOR
 	if (!FPlatformProperties::RequiresCookedData())
 	{
 		for (FShaderType* ShaderType : InShaderTypes)
@@ -453,9 +458,10 @@ void FOpenColorIOShaderMapId::SetShaderDependencies(const TArray<FShaderType*>&
 			}
 		}
 	}
-#endif //WITH_EDITOR
 }
 
+#endif // WITH_EDITOR
+
 bool FOpenColorIOShaderMapId::ContainsShaderType(const FShaderType* ShaderType) const
 {
 	for (int32 TypeIndex = 0; TypeIndex < ShaderTypeDependencies.Num(); TypeIndex++)
diff --git a/Engine/Plugins/Compositing/OpenColorIO/Source/OpenColorIO/Public/OpenColorIOShaderCompilationManager.h b/Engine/Plugins/Compositing/OpenColorIO/Source/OpenColorIO/Public/OpenColorIOShaderCompilationManager.h
--- a/Engine/Plugins/Compositing/OpenColorIO/Source/OpenColorIO/Public/OpenColorIOShaderCompilationManager.h
+++ b/Engine/Plugins/Compositing/OpenColorIO/Source/OpenColorIO/Public/OpenColorIOShaderCompilationManager.h
@@ -2,6 +2,8 @@
 
 #pragma once
 
+#if WITH_EDITOR
+
 #include "CoreMinimal.h"
 #include "ShaderCompiler.h"
 
@@ -108,3 +110,4 @@ private:
 
 extern OPENCOLORIO_API FOpenColorIOShaderCompilationManager GOpenColorIOShaderCompilationManager;
 
+#endif // WITH_EDITOR
diff --git a/Engine/Plugins/Compositing/OpenColorIO/Source/OpenColorIO/Public/OpenColorIOShaderType.h b/Engine/Plugins/Compositing/OpenColorIO/Source/OpenColorIO/Public/OpenColorIOShaderType.h
--- a/Engine/Plugins/Compositing/OpenColorIO/Source/OpenColorIO/Public/OpenColorIOShaderType.h
+++ b/Engine/Plugins/Compositing/OpenColorIO/Source/OpenColorIO/Public/OpenColorIOShaderType.h
@@ -91,6 +91,7 @@ public:
 		check(InTotalPermutationCount == 1);
 	}
 
+#if WITH_EDITOR
 	/**
 	 * Enqueues a compilation for a new shader of this type.
 	 * @param InColorTransform - The ColorTransform to link the shader with.
@@ -113,6 +114,7 @@ public:
 		const FShaderCompileJob& CurrentJob,
 		const FString& InDebugDescription
 		) const;
+#endif // WITH_EDITOR
 
 	/**
 	 * Checks if the shader type should be cached for a particular platform and color transform.
@@ -125,7 +127,7 @@ public:
 		return ShouldCompilePermutation(FOpenColorIOShaderPermutationParameters(GetFName(), InPlatform, InColorTransform));
 	}
 
-
+#if WITH_EDITOR
 protected:
 	/**
 	 * Sets up the environment used to compile an instance of this shader type.
@@ -136,4 +138,5 @@ protected:
 	{
 		ModifyCompilationEnvironment(FOpenColorIOShaderPermutationParameters(GetFName(), InPlatform, InColorTransform), OutEnvironment);
 	}
+#endif // WITH_EDITOR
 };
diff --git a/Engine/Plugins/Compositing/OpenColorIO/Source/OpenColorIO/Public/OpenColorIOShared.h b/Engine/Plugins/Compositing/OpenColorIO/Source/OpenColorIO/Public/OpenColorIOShared.h
--- a/Engine/Plugins/Compositing/OpenColorIO/Source/OpenColorIO/Public/OpenColorIOShared.h
+++ b/Engine/Plugins/Compositing/OpenColorIO/Source/OpenColorIO/Public/OpenColorIOShared.h
@@ -64,7 +64,9 @@ public:
 	~FOpenColorIOShaderMapId()
 	{ }
 
+#if WITH_EDITOR
 	void SetShaderDependencies(const TArray<FShaderType*>& InShaderTypes, EShaderPlatform InShaderPlatform);
+#endif // WITH_EDITOR
 
 	//void Serialize(FArchive& Ar);
 
@@ -94,8 +96,10 @@ public:
 		return !(*this == InReferenceSet);
 	}
 
+#if WITH_EDITOR
 	/** Appends string representations of this Id to a key string. */
 	void AppendKeyString(FString& OutKeyString) const;
+#endif // WITH_EDITOR
 
 	/** Returns true if the requested shader type is a dependency of this shader map Id. */
 	bool ContainsShaderType(const FShaderType* ShaderType) const;
@@ -140,17 +144,20 @@ public:
 	template<typename ShaderType> TShaderRef<ShaderType> GetShader() const { return TShaderRef<ShaderType>(GetContent()->GetShader<ShaderType>(), *this); }
 	TShaderRef<FShader> GetShader(FShaderType* ShaderType) const { return TShaderRef<FShader>(GetContent()->GetShader(ShaderType), *this); }
 
+#if WITH_EDITOR
 	/**
 	 * Attempts to load the shader map for the given color transform from the Derived Data Cache.
 	 * If InOutShaderMap is valid, attempts to load the individual missing shaders instead.
 	 */
 	static void LoadFromDerivedDataCache(const FOpenColorIOTransformResource* InColorTransform, const FOpenColorIOShaderMapId& InShaderMapId, EShaderPlatform InPlatform, TRefCountPtr<FOpenColorIOShaderMap>& InOutShaderMap);
+#endif // WITH_EDITOR
 
 	FOpenColorIOShaderMap();
 
 	// Destructor.
 	~FOpenColorIOShaderMap();
 
+#if WITH_EDITOR
 	/**
 	 * Compiles the shaders for a color transform and caches them in this shader map.
 	 * @param InColorTransform - The color transform to compile shaders for.
@@ -169,6 +176,7 @@ public:
 
 	/** Sorts the incoming compiled jobs into the appropriate OCIO shader maps, and finalizes this shader map so that it can be used for rendering. */
 	bool ProcessCompilationResults(const TArray<FShaderCommonCompileJobPtr>& InCompilationResults, int32& InOutResultIndex, float& InOutTimeBudget);
+#endif // WITH_EDITOR
 
 	/**
 	 * Checks whether the shader map is missing any shader types necessary for the given color transform.
@@ -177,6 +185,7 @@ public:
 	 */
 	bool IsComplete(const FOpenColorIOTransformResource* InColorTransform, bool bSilent);
 
+#if WITH_EDITOR
 	/**
 	 * Checks to see if the shader map is already being compiled for another color transform, and if so
 	 * adds the specified color transform to the list to be applied to once the compile finishes.
@@ -184,6 +193,7 @@ public:
 	 * @return True if the shader map was being compiled and we added ColorTransform to the list to be applied.
 	 */
 	bool TryToAddToExistingCompilationTask(FOpenColorIOTransformResource* InColorTransform);
+#endif // WITH_EDITOR
 
 	/** Builds a list of the shaders in a shader map. */
 	OPENCOLORIO_API  void GetShaderList(TMap<FShaderId, TShaderRef<FShader>>& OutShaders) const;
@@ -194,20 +204,18 @@ public:
 	OPENCOLORIO_API  void AddRef();
 	OPENCOLORIO_API  void Release();
 
-	/**
-	 * Removes all entries in the cache with exceptions based on a shader type
-	 * @param ShaderType - The shader type to flush
-	 */
-	void FlushShadersByShaderType(const FShaderType* ShaderType);
-
+#if WITH_EDITOR
 	/** Removes a ColorTransform from OpenColorIOShaderMapsBeingCompiled. */
 	OPENCOLORIO_API static void RemovePendingColorTransform(FOpenColorIOTransformResource* InColorTransform);
+#endif // WITH_EDITOR
 
 	/** Serializes the shader map. */
 	bool Serialize(FArchive& Ar, bool bInlineShaderResources = true);
 
+#if WITH_EDITOR
 	/** Saves this shader map to the derived data cache. */
 	void SaveToDerivedDataCache();
+#endif // WITH_EDITOR
 
 	/** Registers all shaders that have been loaded in Serialize */
 	//virtual void RegisterSerializedShaders(bool bCooked) override;
@@ -228,9 +236,12 @@ public:
 
 	int32 GetNumRefs() const { return NumRefs; }
 	uint32 GetCompilingId()  { return CompilingId; }
+
+#if WITH_EDITOR
 	static TMap<TRefCountPtr<FOpenColorIOShaderMap>, TArray<FOpenColorIOTransformResource*> > &GetInFlightShaderMaps() { return OpenColorIOShaderMapsBeingCompiled; }
 
 	void SetCompiledSuccessfully(bool bSuccess) { bCompiledSuccessfully = bSuccess; }
+#endif // WITH_EDITOR
 private:
 
 	/**
@@ -246,11 +257,13 @@ private:
 	 */
 	static TArray<FOpenColorIOShaderMap*> AllOpenColorIOShaderMaps;
 
+#if WITH_EDITOR
 	/** Next value for CompilingId. */
 	static uint32 NextCompilingId;
 
 	/** Tracks resources and their shader maps that need to be compiled but whose compilation is being deferred. */
 	static TMap<TRefCountPtr<FOpenColorIOShaderMap>, TArray<FOpenColorIOTransformResource*> > OpenColorIOShaderMapsBeingCompiled;
+#endif // WITH_EDITOR
 
 	/** Uniquely identifies this shader map during compilation, needed for deferred compilation where shaders from multiple shader maps are compiled together. */
 	uint32 CompilingId;
@@ -276,7 +289,9 @@ private:
 
 	uint32 bHasFrozenContent : 1;
 
+#if WITH_EDITOR
 	FShader* ProcessCompilationResultsForSingleJob(const TRefCountPtr<class FShaderCommonCompileJob>& SingleJob, const FSHAHash& InShaderMapHash);
+#endif
 
 	bool IsOpenColorIOShaderComplete(const FOpenColorIOTransformResource* InColorTransform, const FOpenColorIOShaderType* InShaderType, bool bSilent);
 
@@ -343,6 +358,7 @@ public:
 	 */
 	virtual bool IsPersistent() const { return true; }
 
+#if WITH_EDITOR
 	/**
 	 * Cancels all outstanding compilation jobs
 	 */
@@ -352,6 +368,7 @@ public:
 	 * Blocks until compilation has completed. Returns immediately if a compilation is not outstanding.
 	 */
 	void FinishCompilation();
+#endif // WITH_EDITOR
 
 	/**
 	 * Checks if the compilation for this shader is finished
@@ -382,10 +399,12 @@ public:
 	/** Note: SetGameThreadShaderMap must also be called with the same value, but from the game thread. */
 	void SetRenderingThreadShaderMap(FOpenColorIOShaderMap* InShaderMap);
 
+#if WITH_EDITOR
 	void AddCompileId(uint32 InIdentifier) 
 	{
 		OutstandingCompileShaderMapIds.Add(InIdentifier);
 	}
+#endif // WITH_EDITOR
 
 	void SetInlineShaderMap(FOpenColorIOShaderMap* InShaderMap)
 	{
@@ -396,7 +415,9 @@ public:
 		CookedShaderMapId = InShaderMap->GetShaderMapId();
 	}
 
+#if WITH_EDITOR
 	void RemoveOutstandingCompileId(const int32 InOldOutstandingCompileShaderMapId);
+#endif // WITH_EDITOR
 
 	/**
 	* Get OCIO generated source code for the shader
@@ -447,12 +468,12 @@ public:
 	
 	bool IsSame(const FOpenColorIOShaderMapId& InId) const;
 protected:
-
+#if WITH_EDITOR
 	/**
 	 * Fills the passed array with IDs of shader maps unfinished compilation jobs.
 	 */
 	void GetShaderMapIDsWithUnfinishedCompilation(TArray<int32>& OutShaderMapIds);
-
+#endif // WITH_EDITOR
 
 	void SetFeatureLevel(ERHIFeatureLevel::Type InFeatureLevel)
 	{
@@ -483,11 +504,13 @@ private:
 	 */
 	FString ShaderCodeAndConfigHash;
 
+#if WITH_EDITOR
 	/** 
 	 * Contains the compiling id of this shader map when it is being compiled asynchronously. 
 	 * This can be used to access the shader map during async compiling, since GameThreadShaderMap will not have been set yet.
 	 */
 	TArray<int32, TInlineAllocator<1> > OutstandingCompileShaderMapIds;
+#endif // WITH_EDITOR
 
 	/** Feature level that this color transform is representing. */
 	ERHIFeatureLevel::Type FeatureLevel;
@@ -501,6 +524,7 @@ private:
 	uint32 bLoadedCookedShaderMapId : 1;
 	FOpenColorIOShaderMapId CookedShaderMapId;
 
+#if WITH_EDITOR
 	/**
 	 * Compiles this color transform for InPlatform, storing the result in OutShaderMap if the compile was synchronous
 	 */
@@ -516,7 +540,7 @@ private:
 		EShaderPlatform InPlatform,
 		FShaderCompilerEnvironment& OutEnvironment
 		) const;
-
+#endif // WITH_EDITOR
 
 	FString FriendlyName;
 
