OCIO: Optional support for bidirectional display-view.

Inverted view transforms are now supported by our plugin when the option is enabled in the plugin settings, matching the library API.

diff --git a/Engine/Plugins/Compositing/OpenColorIO/Source/OpenColorIO/Private/OpenColorIOConfiguration.cpp b/Engine/Plugins/Compositing/OpenColorIO/Source/OpenColorIO/Private/OpenColorIOConfiguration.cpp
--- a/Engine/Plugins/Compositing/OpenColorIO/Source/OpenColorIO/Private/OpenColorIOConfiguration.cpp
+++ b/Engine/Plugins/Compositing/OpenColorIO/Source/OpenColorIO/Private/OpenColorIOConfiguration.cpp
@@ -11,6 +11,7 @@
 #include "Modules/ModuleManager.h"
 #include "OpenColorIOColorTransform.h"
 #include "OpenColorIOModule.h"
+#include "OpenColorIOSettings.h"
 #include "TextureResource.h"
 #include "Widgets/Notifications/SNotificationList.h"
 
@@ -209,6 +210,8 @@ void UOpenColorIOConfiguration::ReloadExistingColorspaces()
 		DesiredDisplayViews.Add(ExistingDisplayView);
 	}
 
+	const UOpenColorIOSettings* Settings = GetDefault<UOpenColorIOSettings>();
+
 	// Genereate new shaders.
 	for (int32 indexTop = 0; indexTop < DesiredColorSpaces.Num(); ++indexTop)
 	{
@@ -224,8 +227,12 @@ void UOpenColorIOConfiguration::ReloadExistingColorspaces()
 
 		for (const FOpenColorIODisplayView& DisplayView : DesiredDisplayViews)
 		{
-			// note: we only support display-view transforms in the forward direction currently.
 			CreateColorTransform(TopColorSpace.ColorSpaceName, DisplayView.Display, DisplayView.View, EOpenColorIOViewTransformDirection::Forward);
+
+			if (Settings->bSupportInverseViewTransforms)
+			{
+				CreateColorTransform(TopColorSpace.ColorSpaceName, DisplayView.Display, DisplayView.View, EOpenColorIOViewTransformDirection::Inverse);
+			}
 		}
 	}
 #endif
@@ -461,6 +468,8 @@ void UOpenColorIOConfiguration::PreSave(FObjectPreSaveContext SaveContext)
 
 void UOpenColorIOConfiguration::PostEditChangeProperty(FPropertyChangedEvent& PropertyChangedEvent)
 {
+	const UOpenColorIOSettings* Settings = GetDefault<UOpenColorIOSettings>();
+
 	if (PropertyChangedEvent.MemberProperty && PropertyChangedEvent.MemberProperty->GetFName() == GET_MEMBER_NAME_CHECKED(UOpenColorIOConfiguration, ConfigurationFile))
 	{
 		LoadConfigurationFile();
@@ -483,8 +492,12 @@ void UOpenColorIOConfiguration::PostEditChangeProperty(FPropertyChangedEvent& Pr
 
 				for (const FOpenColorIODisplayView& DisplayView : DesiredDisplayViews)
 				{
-					// note: we only support display-view transforms in the forward direction currently.
 					CreateColorTransform(TopColorSpace.ColorSpaceName, DisplayView.Display, DisplayView.View, EOpenColorIOViewTransformDirection::Forward);
+
+					if (Settings->bSupportInverseViewTransforms)
+					{
+						CreateColorTransform(TopColorSpace.ColorSpaceName, DisplayView.Display, DisplayView.View, EOpenColorIOViewTransformDirection::Inverse);
+					}
 				}
 			}
 		}
@@ -502,8 +515,12 @@ void UOpenColorIOConfiguration::PostEditChangeProperty(FPropertyChangedEvent& Pr
 			{
 				for (const FOpenColorIOColorSpace& SourceColorSpace : DesiredColorSpaces)
 				{
-					// note: we only support display-view transforms in the forward direction currently.
 					CreateColorTransform(SourceColorSpace.ColorSpaceName, DisplayView.Display, DisplayView.View, EOpenColorIOViewTransformDirection::Forward);
+
+					if (Settings->bSupportInverseViewTransforms)
+					{
+						CreateColorTransform(SourceColorSpace.ColorSpaceName, DisplayView.Display, DisplayView.View, EOpenColorIOViewTransformDirection::Inverse);
+					}
 				}
 			}
 		}
diff --git a/Engine/Plugins/Compositing/OpenColorIO/Source/OpenColorIO/Private/OpenColorIOSettings.cpp b/Engine/Plugins/Compositing/OpenColorIO/Source/OpenColorIO/Private/OpenColorIOSettings.cpp
--- a/Engine/Plugins/Compositing/OpenColorIO/Source/OpenColorIO/Private/OpenColorIOSettings.cpp
+++ b/Engine/Plugins/Compositing/OpenColorIO/Source/OpenColorIO/Private/OpenColorIOSettings.cpp
@@ -5,6 +5,7 @@
 UOpenColorIOSettings::UOpenColorIOSettings()
 	: bUseLegacyProcessor(false)
 	, bUse32fLUT(false)
+	, bSupportInverseViewTransforms(false)
 {
 
 }
diff --git a/Engine/Plugins/Compositing/OpenColorIO/Source/OpenColorIO/Public/OpenColorIOSettings.h b/Engine/Plugins/Compositing/OpenColorIO/Source/OpenColorIO/Public/OpenColorIOSettings.h
--- a/Engine/Plugins/Compositing/OpenColorIO/Source/OpenColorIO/Public/OpenColorIOSettings.h
+++ b/Engine/Plugins/Compositing/OpenColorIO/Source/OpenColorIO/Public/OpenColorIOSettings.h
@@ -41,4 +41,10 @@ public:
 		ToolTip = "Whether to create lookup table texture resources in 32-bit float format (higher performance requirements).",
 		ConfigRestartRequired = true))
 	uint8 bUse32fLUT : 1;
+
+	UPROPERTY(config, EditAnywhere, Category = Transform, meta = (
+		DisplayName = "Support inverse view transforms",
+		ToolTip = "Whether inverse view transforms are cached and supported. Disabled by default, to minimize the number of transform combinations.",
+		ConfigRestartRequired = true))
+	uint8 bSupportInverseViewTransforms : 1;
 };
diff --git a/Engine/Plugins/Compositing/OpenColorIO/Source/OpenColorIOEditor/Private/OpenColorIOColorSpaceConversionCustomization.cpp b/Engine/Plugins/Compositing/OpenColorIO/Source/OpenColorIOEditor/Private/OpenColorIOColorSpaceConversionCustomization.cpp
--- a/Engine/Plugins/Compositing/OpenColorIO/Source/OpenColorIOEditor/Private/OpenColorIOColorSpaceConversionCustomization.cpp
+++ b/Engine/Plugins/Compositing/OpenColorIO/Source/OpenColorIOEditor/Private/OpenColorIOColorSpaceConversionCustomization.cpp
@@ -10,6 +10,7 @@
 #include "IPropertyUtilities.h"
 #include "OpenColorIOColorSpace.h"
 #include "OpenColorIOConfiguration.h"
+#include "OpenColorIOSettings.h"
 #include "PropertyHandle.h"
 #include "Widgets/SOpenColorIOColorSpacePicker.h"
 #include "Widgets/Text/STextBlock.h"
@@ -97,8 +98,8 @@ void FOpenColorIOColorConversionSettingsCustomization::CustomizeChildren(TShared
 
 				ChildHandle->SetOnPropertyValueChanged(FSimpleDelegate::CreateLambda([this]
 					{
-						TransformSourcePicker->SetConfiguration(ColorSpaceConversion->ConfigurationSource);
-						TransformDestinationPicker->SetConfiguration(ColorSpaceConversion->ConfigurationSource);
+						TransformPicker[OCIO_Src]->SetConfiguration(ColorSpaceConversion->ConfigurationSource);
+						TransformPicker[OCIO_Dst]->SetConfiguration(ColorSpaceConversion->ConfigurationSource);
 					}));
 			}
 			else if (ChildHandle->GetProperty()->GetFName() == GET_MEMBER_NAME_CHECKED(FOpenColorIOColorConversionSettings, SourceColorSpace))
@@ -113,6 +114,10 @@ void FOpenColorIOColorConversionSettingsCustomization::CustomizeChildren(TShared
 			{
 				DestinationDisplayViewProperty = ChildHandle;
 			}
+			else if (ChildHandle->GetProperty()->GetFName() == GET_MEMBER_NAME_CHECKED(FOpenColorIOColorConversionSettings, DisplayViewDirection))
+			{
+				DisplayViewDirectionProperty = ChildHandle;
+			}
 		}
 
 		// Note: ChildHandle->SetOnPropertyValueChanged isn't automatically called on parent changes, so we explicitely reset the configuration here.
@@ -122,37 +127,29 @@ void FOpenColorIOColorConversionSettingsCustomization::CustomizeChildren(TShared
 			ParentHandle->SetOnPropertyResetToDefault(FSimpleDelegate::CreateSP(this, &FOpenColorIOColorConversionSettingsCustomization::OnConfigurationReset));
 		}
 
-		TransformSourcePicker = SNew(SOpenColorIOColorSpacePicker)
+		ApplyConfigurationToSelection();
+
+		// Source color space picker widget
+		TransformPicker[OCIO_Src] =
+			SNew(SOpenColorIOColorSpacePicker)
 			.Config(ColorSpaceConversion->ConfigurationSource)
-			.InitialColorSpace(ColorSpaceConversion->SourceColorSpace)
-			.RestrictedColor(ColorSpaceConversion->DestinationColorSpace)
-			.InitialDisplayView(FOpenColorIODisplayView())
+			.InitialColorSpace(TransformSelection[OCIO_Src].ColorSpace)
+			.RestrictedColor(TransformSelection[OCIO_Dst].ColorSpace)
+			.InitialDisplayView(TransformSelection[OCIO_Src].DisplayView)
+			.RestrictedDisplayView(TransformSelection[OCIO_Dst].DisplayView)
 			.IsDestination(false)
-			.OnColorSpaceChanged(FOnColorSpaceChanged::CreateLambda([this](const FOpenColorIOColorSpace& NewColorSpace, const FOpenColorIODisplayView& NewDisplayView)
-				{
-					check(SourceColorSpaceProperty.IsValid());
+			.OnColorSpaceChanged(FOnColorSpaceChanged::CreateSP(this, &FOpenColorIOColorConversionSettingsCustomization::OnSourceColorSpaceChanged));
 
-					UpdateColorSettingsStructProperty(SourceColorSpaceProperty, &NewColorSpace, &ColorSpaceConversion->SourceColorSpace);
-					
-					TransformDestinationPicker->SetRestrictedColorSpace(NewColorSpace);
-				}));
-
-		TransformDestinationPicker = SNew(SOpenColorIOColorSpacePicker)
+		// Destination color space picker widget
+		TransformPicker[OCIO_Dst] =
+			SNew(SOpenColorIOColorSpacePicker)
 			.Config(ColorSpaceConversion->ConfigurationSource)
-			.InitialColorSpace(ColorSpaceConversion->DestinationColorSpace)
-			.RestrictedColor(ColorSpaceConversion->SourceColorSpace)
-			.InitialDisplayView(ColorSpaceConversion->DestinationDisplayView)
+			.InitialColorSpace(TransformSelection[OCIO_Dst].ColorSpace)
+			.RestrictedColor(TransformSelection[OCIO_Src].ColorSpace)
+			.InitialDisplayView(TransformSelection[OCIO_Dst].DisplayView)
+			.RestrictedDisplayView(TransformSelection[OCIO_Src].DisplayView)
 			.IsDestination(true)
-			.OnColorSpaceChanged(FOnColorSpaceChanged::CreateLambda([this](const FOpenColorIOColorSpace& NewColorSpace, const FOpenColorIODisplayView& NewDisplayView)
-				{
-					check(DestinationColorSpaceProperty.IsValid());
-					check(DestinationDisplayViewProperty.IsValid());
-
-					UpdateColorSettingsStructProperty(DestinationColorSpaceProperty, &NewColorSpace, &ColorSpaceConversion->DestinationColorSpace);
-					UpdateColorSettingsStructProperty(DestinationDisplayViewProperty, &NewDisplayView, &ColorSpaceConversion->DestinationDisplayView);
-					
-					TransformSourcePicker->SetRestrictedColorSpace(NewColorSpace);
-				}));
+			.OnColorSpaceChanged(FOnColorSpaceChanged::CreateSP(this, &FOpenColorIOColorConversionSettingsCustomization::OnDestinationColorSpaceChanged));
 
 		// Source color space picker widget
 		StructBuilder.AddCustomRow(LOCTEXT("TransformSource", "Transform Source"))
@@ -166,7 +163,7 @@ void FOpenColorIOColorConversionSettingsCustomization::CustomizeChildren(TShared
 		.ValueContent()
 		.MaxDesiredWidth(512)
 		[
-			TransformSourcePicker.ToSharedRef()
+			TransformPicker[OCIO_Src].ToSharedRef()
 		];
 
 		StructBuilder.AddCustomRow(LOCTEXT("TransformDestination", "Transform Destination"))
@@ -180,17 +177,78 @@ void FOpenColorIOColorConversionSettingsCustomization::CustomizeChildren(TShared
 		.ValueContent()
 		.MaxDesiredWidth(512)
 		[
-			TransformDestinationPicker.ToSharedRef()
+			TransformPicker[OCIO_Dst].ToSharedRef()
 		];
 	}
 
 	InStructPropertyHandle->MarkHiddenByCustomization();
 }
 
+void FOpenColorIOColorConversionSettingsCustomization::OnSourceColorSpaceChanged(const FOpenColorIOColorSpace& NewColorSpace, const FOpenColorIODisplayView& NewDisplayView)
+{
+	TransformSelection[OCIO_Src].ColorSpace = NewColorSpace;
+	TransformSelection[OCIO_Src].DisplayView = NewDisplayView;
+
+	TransformPicker[OCIO_Dst]->SetRestrictions(NewColorSpace, NewDisplayView);
+
+	ApplySelectionToConfiguration();
+}
+
+void FOpenColorIOColorConversionSettingsCustomization::OnDestinationColorSpaceChanged(const FOpenColorIOColorSpace& NewColorSpace, const FOpenColorIODisplayView& NewDisplayView)
+{
+	TransformSelection[OCIO_Dst].ColorSpace = NewColorSpace;
+	TransformSelection[OCIO_Dst].DisplayView = NewDisplayView;
+
+	TransformPicker[OCIO_Src]->SetRestrictions(NewColorSpace, NewDisplayView);
+
+	ApplySelectionToConfiguration();
+}
+
+void FOpenColorIOColorConversionSettingsCustomization::ApplyConfigurationToSelection()
+{
+	if (ColorSpaceConversion->DisplayViewDirection == EOpenColorIOViewTransformDirection::Forward)
+	{
+		TransformSelection[OCIO_Src].ColorSpace = ColorSpaceConversion->SourceColorSpace;
+		TransformSelection[OCIO_Src].DisplayView.Reset();
+		TransformSelection[OCIO_Dst].ColorSpace = ColorSpaceConversion->DestinationColorSpace;
+		TransformSelection[OCIO_Dst].DisplayView = ColorSpaceConversion->DestinationDisplayView;
+	}
+	else
+	{
+		TransformSelection[OCIO_Src].ColorSpace.Reset();
+		TransformSelection[OCIO_Src].DisplayView = ColorSpaceConversion->DestinationDisplayView;
+		TransformSelection[OCIO_Dst].ColorSpace = ColorSpaceConversion->SourceColorSpace;
+		TransformSelection[OCIO_Dst].DisplayView.Reset();
+	}
+}
+
+void FOpenColorIOColorConversionSettingsCustomization::ApplySelectionToConfiguration()
+{
+	check(SourceColorSpaceProperty.IsValid() && DestinationColorSpaceProperty.IsValid() && DestinationDisplayViewProperty.IsValid() && DisplayViewDirectionProperty.IsValid());
+
+	FOpenColorIOColorConversionSettings& ConversionSettings = *ColorSpaceConversion;
+
+	if (TransformSelection[OCIO_Src].DisplayView.IsValid())
+	{
+		FOpenColorIOColorSpace EmptyCS = {};
+		UpdateColorSettingsStructProperty(SourceColorSpaceProperty, &TransformSelection[OCIO_Dst].ColorSpace, &ColorSpaceConversion->SourceColorSpace);
+		UpdateColorSettingsStructProperty(DestinationColorSpaceProperty, &EmptyCS, &ColorSpaceConversion->DestinationColorSpace);
+		UpdateColorSettingsStructProperty(DestinationDisplayViewProperty, &TransformSelection[OCIO_Src].DisplayView, &ColorSpaceConversion->DestinationDisplayView);
+		DisplayViewDirectionProperty->SetValue((uint8)EOpenColorIOViewTransformDirection::Inverse);
+	}
+	else
+	{
+		UpdateColorSettingsStructProperty(SourceColorSpaceProperty, &TransformSelection[OCIO_Src].ColorSpace, &ColorSpaceConversion->SourceColorSpace);
+		UpdateColorSettingsStructProperty(DestinationColorSpaceProperty, &TransformSelection[OCIO_Dst].ColorSpace, &ColorSpaceConversion->DestinationColorSpace);
+		UpdateColorSettingsStructProperty(DestinationDisplayViewProperty, &TransformSelection[OCIO_Dst].DisplayView, &ColorSpaceConversion->DestinationDisplayView);
+		DisplayViewDirectionProperty->SetValue((uint8)EOpenColorIOViewTransformDirection::Forward);
+	}
+}
+
 void FOpenColorIOColorConversionSettingsCustomization::OnConfigurationReset()
 {
-	TransformSourcePicker->SetConfiguration(nullptr);
-	TransformDestinationPicker->SetConfiguration(nullptr);
+	TransformPicker[OCIO_Src]->SetConfiguration(nullptr);
+	TransformPicker[OCIO_Dst]->SetConfiguration(nullptr);
 }
 
 #undef LOCTEXT_NAMESPACE
diff --git a/Engine/Plugins/Compositing/OpenColorIO/Source/OpenColorIOEditor/Private/OpenColorIOColorSpaceConversionCustomization.h b/Engine/Plugins/Compositing/OpenColorIO/Source/OpenColorIOEditor/Private/OpenColorIOColorSpaceConversionCustomization.h
--- a/Engine/Plugins/Compositing/OpenColorIO/Source/OpenColorIOEditor/Private/OpenColorIOColorSpaceConversionCustomization.h
+++ b/Engine/Plugins/Compositing/OpenColorIO/Source/OpenColorIOEditor/Private/OpenColorIOColorSpaceConversionCustomization.h
@@ -3,11 +3,11 @@
 #pragma once
 
 #include "CoreMinimal.h"
+#include "OpenColorIOColorSpace.h"
 #include "IPropertyTypeCustomization.h"
+#include "Widgets/SOpenColorIOColorSpacePicker.h"
 #include "Widgets/SWidget.h"
 
-class SOpenColorIOColorSpacePicker;
-
 /**
  * Implements a details view customization for the FOpenColorIOColorConversionSettings
  */
@@ -23,22 +23,35 @@ public:
 	virtual void CustomizeHeader(TSharedRef<class IPropertyHandle> InPropertyHandle, class FDetailWidgetRow& HeaderRow, IPropertyTypeCustomizationUtils& PropertyTypeCustomizationUtils) override;
 	virtual void CustomizeChildren(TSharedRef<class IPropertyHandle> InPropertyHandle, class IDetailChildrenBuilder& StructBuilder, IPropertyTypeCustomizationUtils& PropertyTypeCustomizationUtils) override;
 
+	void OnSourceColorSpaceChanged(const FOpenColorIOColorSpace& NewColorSpace, const FOpenColorIODisplayView& NewDisplayView);
+	void OnDestinationColorSpaceChanged(const FOpenColorIOColorSpace& NewColorSpace, const FOpenColorIODisplayView& NewDisplayView);
+
+protected:
+	/** Read configuration settings into local selection objects.*/
+	void ApplyConfigurationToSelection();
+	/** Apply transform selections onto the configuration object.*/
+	void ApplySelectionToConfiguration();
+
 private:
 	/** Callback to reset the configuration in the transform source/destination pickers. */
 	void OnConfigurationReset();
 
-	/** Pointer to the ColorConversion struct member SourceColorSpace property handle. */
+	/** Pointer to the struct SourceColorSpace property handle. */
 	TSharedPtr<IPropertyHandle> SourceColorSpaceProperty;
 	
-	/** Pointer to the ColorConversion struct member DestinationColorSpace property handle. */
+	/** Pointer to the struct DestinationColorSpace property handle. */
 	TSharedPtr<IPropertyHandle> DestinationColorSpaceProperty;
 
-	/** Pointer to the ColorConversion struct member DestinationColorSpace property handle. */
+	/** Pointer to the struct DestinationDisplayView property handle. */
 	TSharedPtr<IPropertyHandle> DestinationDisplayViewProperty;
 
+	/** Pointer to the struct DisplayViewDirection property handle. */
+	TSharedPtr<IPropertyHandle> DisplayViewDirectionProperty;
+
 	/** ColorSpace pickers reference to update them when config asset is changed */
-	TSharedPtr<SOpenColorIOColorSpacePicker> TransformSourcePicker = nullptr;
-	TSharedPtr<SOpenColorIOColorSpacePicker> TransformDestinationPicker = nullptr;
+	TStaticArray<TSharedPtr<SOpenColorIOColorSpacePicker>, 2> TransformPicker;
+	/** Intermediate selection objects. This is done to allow inverted display-view selections (when enabled in settings). */
+	TStaticArray<FOpenColorIOPickerSelection, 2> TransformSelection;
 
 	/** Raw pointer to the conversion settings struct. */
 	struct FOpenColorIOColorConversionSettings* ColorSpaceConversion = nullptr;
diff --git a/Engine/Plugins/Compositing/OpenColorIO/Source/OpenColorIOEditor/Private/Widgets/SOpenColorIOColorSpacePicker.cpp b/Engine/Plugins/Compositing/OpenColorIO/Source/OpenColorIOEditor/Private/Widgets/SOpenColorIOColorSpacePicker.cpp
--- a/Engine/Plugins/Compositing/OpenColorIO/Source/OpenColorIOEditor/Private/Widgets/SOpenColorIOColorSpacePicker.cpp
+++ b/Engine/Plugins/Compositing/OpenColorIO/Source/OpenColorIOEditor/Private/Widgets/SOpenColorIOColorSpacePicker.cpp
@@ -5,6 +5,7 @@
 #include "Styling/AppStyle.h"
 #include "Framework/MultiBox/MultiBoxBuilder.h"
 #include "OpenColorIOConfiguration.h"
+#include "OpenColorIOSettings.h"
 #include "SResetToDefaultMenu.h"
 #include "Widgets/Images/SImage.h"
 #include "Widgets/Input/SButton.h"
@@ -21,6 +22,7 @@ void SOpenColorIOColorSpacePicker::Construct(const FArguments& InArgs)
 	ColorSpaceSelection = InArgs._InitialColorSpace;
 	RestrictedColorSpace = InArgs._RestrictedColor;
 	DisplayViewSelection = InArgs._InitialDisplayView;
+	RestrictedDisplayView = InArgs._RestrictedDisplayView;
 	bIsDestination = InArgs._IsDestination;
 	OnColorSpaceChanged = InArgs._OnColorSpaceChanged;
 
@@ -85,9 +87,10 @@ void SOpenColorIOColorSpacePicker::SetConfiguration(TWeakObjectPtr<UOpenColorIOC
 	SetCurrentColorSpace(FOpenColorIOColorSpace());
 }
 
-void SOpenColorIOColorSpacePicker::SetRestrictedColorSpace(const FOpenColorIOColorSpace& InRestrictedColorSpace)
+void SOpenColorIOColorSpacePicker::SetRestrictions(const FOpenColorIOColorSpace& InRestrictedColorSpace, const FOpenColorIODisplayView& InRestricedDisplayView)
 {
 	RestrictedColorSpace = InRestrictedColorSpace;
+	RestrictedDisplayView = InRestricedDisplayView;
 }
 
 void SOpenColorIOColorSpacePicker::SetCurrentColorSpace(const FOpenColorIOColorSpace& NewColorSpace)
@@ -124,6 +127,8 @@ TSharedRef<SWidget> SOpenColorIOColorSpacePicker::HandleColorSpaceComboButtonMen
 {
 	if (UOpenColorIOConfiguration* ConfigurationObject = Configuration.Get())
 	{
+		const UOpenColorIOSettings* Settings = GetDefault<UOpenColorIOSettings>();
+
 		// generate menu
 		const bool bShouldCloseWindowAfterClosing = false;
 		FMenuBuilder MenuBuilder(bShouldCloseWindowAfterClosing, nullptr);
@@ -171,7 +176,7 @@ TSharedRef<SWidget> SOpenColorIOColorSpacePicker::HandleColorSpaceComboButtonMen
 		}
 		MenuBuilder.EndSection();
 
-		if (bIsDestination)
+		if (bIsDestination || Settings->bSupportInverseViewTransforms)
 		{
 			MenuBuilder.BeginSection("AvailableDisplayViews", LOCTEXT("AvailableDisplayViews", "Available Display-Views"));
 			{
@@ -180,7 +185,7 @@ TSharedRef<SWidget> SOpenColorIOColorSpacePicker::HandleColorSpaceComboButtonMen
 				for (int32 i = 0; i < ConfigurationObject->DesiredDisplayViews.Num(); ++i)
 				{
 					const FOpenColorIODisplayView& DisplayView = ConfigurationObject->DesiredDisplayViews[i];
-					if (DisplayView == DisplayViewSelection || !DisplayView.IsValid())
+					if (DisplayView == RestrictedDisplayView || !DisplayView.IsValid())
 					{
 						continue;
 					}
diff --git a/Engine/Plugins/Compositing/OpenColorIO/Source/OpenColorIOEditor/Private/Widgets/SOpenColorIOColorSpacePicker.h b/Engine/Plugins/Compositing/OpenColorIO/Source/OpenColorIOEditor/Private/Widgets/SOpenColorIOColorSpacePicker.h
--- a/Engine/Plugins/Compositing/OpenColorIO/Source/OpenColorIOEditor/Private/Widgets/SOpenColorIOColorSpacePicker.h
+++ b/Engine/Plugins/Compositing/OpenColorIO/Source/OpenColorIOEditor/Private/Widgets/SOpenColorIOColorSpacePicker.h
@@ -18,6 +18,20 @@ class UOpenColorIOConfiguration;
 
 DECLARE_DELEGATE_TwoParams(FOnColorSpaceChanged, const FOpenColorIOColorSpace& /*ColorSpace*/, const FOpenColorIODisplayView& /*DisplayView*/);
 
+/** Pair of OpenColorIO transform selection objects. */
+struct FOpenColorIOPickerSelection
+{
+	FOpenColorIOColorSpace ColorSpace;
+	FOpenColorIODisplayView DisplayView;
+};
+
+/** OpenColorIO transformation source or destination. */
+enum EOpenColorIOTransformDomain : uint32
+{
+	OCIO_Src = 0,
+	OCIO_Dst = 1,
+};
+
 class SOpenColorIOColorSpacePicker : public SCompoundWidget
 {
 public:
@@ -26,6 +40,7 @@ public:
 		SLATE_ARGUMENT(FOpenColorIOColorSpace, InitialColorSpace)
 		SLATE_ARGUMENT(FOpenColorIOColorSpace, RestrictedColor)
 		SLATE_ARGUMENT(FOpenColorIODisplayView, InitialDisplayView)
+		SLATE_ARGUMENT(FOpenColorIODisplayView, RestrictedDisplayView)
 		SLATE_ARGUMENT(bool, IsDestination)
 		SLATE_EVENT(FOnColorSpaceChanged, OnColorSpaceChanged)
 	SLATE_END_ARGS()
@@ -37,7 +52,7 @@ public:
 	void SetConfiguration(TWeakObjectPtr<UOpenColorIOConfiguration> NewConfiguration);
 
 	/** Update restricted color space for this picker */
-	void SetRestrictedColorSpace(const FOpenColorIOColorSpace& RestrictedColorSpace);
+	void SetRestrictions(const FOpenColorIOColorSpace& RestrictedColorSpace, const FOpenColorIODisplayView& RestricedDisplayView);
 
 protected:
 	
@@ -60,6 +75,7 @@ protected:
 	FOpenColorIOColorSpace ColorSpaceSelection;
 	FOpenColorIOColorSpace RestrictedColorSpace;
 	FOpenColorIODisplayView DisplayViewSelection;
+	FOpenColorIODisplayView RestrictedDisplayView;
 	FOnColorSpaceChanged OnColorSpaceChanged;
 	bool bIsDestination = false;
 };
diff --git a/Engine/Plugins/Compositing/OpenColorIO/Source/OpenColorIOEditor/Private/Widgets/SOpenColorIODisplay.cpp b/Engine/Plugins/Compositing/OpenColorIO/Source/OpenColorIOEditor/Private/Widgets/SOpenColorIODisplay.cpp
--- a/Engine/Plugins/Compositing/OpenColorIO/Source/OpenColorIOEditor/Private/Widgets/SOpenColorIODisplay.cpp
+++ b/Engine/Plugins/Compositing/OpenColorIO/Source/OpenColorIOEditor/Private/Widgets/SOpenColorIODisplay.cpp
@@ -61,28 +61,33 @@ void SOpenColorIODisplay::Construct(const FArguments& InArgs)
 		EUserInterfaceActionType::None
 	));
 
+	ApplyConfigurationToSelection();
+
 	// Source color space picker widget
-	TransformSourcePicker =
+	TransformPicker[OCIO_Src] =
 		SNew(SOpenColorIOColorSpacePicker)
 		.Config(Configuration.ColorConfiguration.ConfigurationSource)
-		.InitialColorSpace(Configuration.ColorConfiguration.SourceColorSpace)
-		.RestrictedColor(Configuration.ColorConfiguration.DestinationColorSpace)
-		.InitialDisplayView(FOpenColorIODisplayView())
+		.InitialColorSpace(TransformSelection[OCIO_Src].ColorSpace)
+		.RestrictedColor(TransformSelection[OCIO_Dst].ColorSpace)
+		.InitialDisplayView(TransformSelection[OCIO_Src].DisplayView)
+		.RestrictedDisplayView(TransformSelection[OCIO_Dst].DisplayView)
 		.IsDestination(false)
 		.OnColorSpaceChanged(FOnColorSpaceChanged::CreateSP(this, &SOpenColorIODisplay::OnSourceColorSpaceChanged));
 
-	// Source color space picker widget
-	TransformDestinationPicker =
+	// Destination color space picker widget
+	TransformPicker[OCIO_Dst] =
 		SNew(SOpenColorIOColorSpacePicker)
 		.Config(Configuration.ColorConfiguration.ConfigurationSource)
-		.InitialColorSpace(Configuration.ColorConfiguration.DestinationColorSpace)
-		.RestrictedColor(Configuration.ColorConfiguration.SourceColorSpace)
-		.InitialDisplayView(Configuration.ColorConfiguration.DestinationDisplayView)
+		.InitialColorSpace(TransformSelection[OCIO_Dst].ColorSpace)
+		.RestrictedColor(TransformSelection[OCIO_Src].ColorSpace)
+		.InitialDisplayView(TransformSelection[OCIO_Dst].DisplayView)
+		.RestrictedDisplayView(TransformSelection[OCIO_Src].DisplayView)
 		.IsDestination(true)
 		.OnColorSpaceChanged(FOnColorSpaceChanged::CreateSP(this, &SOpenColorIODisplay::OnDestinationColorSpaceChanged));
 	
-	Section.AddEntry(FToolMenuEntry::InitWidget("SourceColor", TransformSourcePicker.ToSharedRef(), FText::GetEmpty(), true, false));
-	Section.AddEntry(FToolMenuEntry::InitWidget("DestinationColor", TransformDestinationPicker.ToSharedRef(), FText::GetEmpty(), true, false));
+	// Note: FText::GetEmpty() would previously result in incorrect indentation.
+	Section.AddEntry(FToolMenuEntry::InitWidget("SourceColor", TransformPicker[OCIO_Src].ToSharedRef(), FText::FromString(TEXT(" ")), false, false));
+	Section.AddEntry(FToolMenuEntry::InitWidget("DestinationColor", TransformPicker[OCIO_Dst].ToSharedRef(), FText::FromString(TEXT(" ")), false, false));
 
 	TSharedRef<SWidget> Widget = UToolMenus::Get()->GenerateWidget(Menu);
 	ChildSlot
@@ -149,8 +154,8 @@ void SOpenColorIODisplay::OnConfigSelected(const FAssetData& AssetData)
 
 	UOpenColorIOConfiguration* NewConfig = Cast<UOpenColorIOConfiguration>(AssetData.GetAsset());
 	Configuration.ColorConfiguration.ConfigurationSource = NewConfig;
-	TransformSourcePicker->SetConfiguration(NewConfig);
-	TransformDestinationPicker->SetConfiguration(NewConfig);
+	TransformPicker[OCIO_Src]->SetConfiguration(NewConfig);
+	TransformPicker[OCIO_Dst]->SetConfiguration(NewConfig);
 
 	OnDisplayConfigurationChanged.ExecuteIfBound(Configuration);
 }
@@ -173,21 +178,65 @@ bool SOpenColorIODisplay::GetDisplayConfigurationState()
 	return Configuration.bIsEnabled;
 }
 
-void SOpenColorIODisplay::OnSourceColorSpaceChanged(const FOpenColorIOColorSpace& NewColorSpace, const FOpenColorIODisplayView&)
+void SOpenColorIODisplay::OnSourceColorSpaceChanged(const FOpenColorIOColorSpace& NewColorSpace, const FOpenColorIODisplayView& NewDisplayView)
 {
-	//When color space changes, assign it, update restriction on the other picker and let know our listeners we changed
-	Configuration.ColorConfiguration.SourceColorSpace = NewColorSpace;
-	TransformDestinationPicker->SetRestrictedColorSpace(NewColorSpace);
+	TransformSelection[OCIO_Src].ColorSpace = NewColorSpace;
+	TransformSelection[OCIO_Src].DisplayView = NewDisplayView;
 
-	OnDisplayConfigurationChanged.ExecuteIfBound(Configuration);
+	TransformPicker[OCIO_Dst]->SetRestrictions(NewColorSpace, NewDisplayView);
+
+	ApplySelectionToConfiguration();
 }
 
 void SOpenColorIODisplay::OnDestinationColorSpaceChanged(const FOpenColorIOColorSpace& NewColorSpace, const FOpenColorIODisplayView& NewDisplayView)
 {
-	Configuration.ColorConfiguration.DestinationColorSpace = NewColorSpace;
-	TransformSourcePicker->SetRestrictedColorSpace(NewColorSpace);
+	TransformSelection[OCIO_Dst].ColorSpace = NewColorSpace;
+	TransformSelection[OCIO_Dst].DisplayView = NewDisplayView;
+
+	TransformPicker[OCIO_Src]->SetRestrictions(NewColorSpace, NewDisplayView);
+
+	ApplySelectionToConfiguration();
+}
+
+void SOpenColorIODisplay::ApplyConfigurationToSelection()
+{
+	const FOpenColorIOColorConversionSettings& ConversionSettings = Configuration.ColorConfiguration;
+
+	if (ConversionSettings.DisplayViewDirection == EOpenColorIOViewTransformDirection::Forward)
+	{
+		TransformSelection[OCIO_Src].ColorSpace = ConversionSettings.SourceColorSpace;
+		TransformSelection[OCIO_Src].DisplayView.Reset();
+		TransformSelection[OCIO_Dst].ColorSpace = ConversionSettings.DestinationColorSpace;
+		TransformSelection[OCIO_Dst].DisplayView = ConversionSettings.DestinationDisplayView;
+	}
+	else
+	{
+		TransformSelection[OCIO_Src].ColorSpace.Reset();
+		TransformSelection[OCIO_Src].DisplayView = ConversionSettings.DestinationDisplayView;
+		TransformSelection[OCIO_Dst].ColorSpace = ConversionSettings.SourceColorSpace;
+		TransformSelection[OCIO_Dst].DisplayView.Reset();
+	}
+}
+
+void SOpenColorIODisplay::ApplySelectionToConfiguration()
+{
+	FOpenColorIOColorConversionSettings& ConversionSettings = Configuration.ColorConfiguration;
+
+	if (TransformSelection[OCIO_Src].DisplayView.IsValid())
+	{
+		ConversionSettings.SourceColorSpace = TransformSelection[OCIO_Dst].ColorSpace;
+		ConversionSettings.DestinationColorSpace.Reset();
+		ConversionSettings.DestinationDisplayView = TransformSelection[OCIO_Src].DisplayView;
+		ConversionSettings.DisplayViewDirection = EOpenColorIOViewTransformDirection::Inverse;
+	}
+	else
+	{
+		ConversionSettings.SourceColorSpace = TransformSelection[OCIO_Src].ColorSpace;
+		ConversionSettings.DestinationColorSpace = TransformSelection[OCIO_Dst].ColorSpace;
+		ConversionSettings.DestinationDisplayView = TransformSelection[OCIO_Dst].DisplayView;
+		ConversionSettings.DisplayViewDirection = EOpenColorIOViewTransformDirection::Forward;
+	}
 
-	Configuration.ColorConfiguration.DestinationDisplayView = NewDisplayView;
 	OnDisplayConfigurationChanged.ExecuteIfBound(Configuration);
 }
 
diff --git a/Engine/Plugins/Compositing/OpenColorIO/Source/OpenColorIOEditor/Private/Widgets/SOpenColorIODisplay.h b/Engine/Plugins/Compositing/OpenColorIO/Source/OpenColorIOEditor/Private/Widgets/SOpenColorIODisplay.h
--- a/Engine/Plugins/Compositing/OpenColorIO/Source/OpenColorIOEditor/Private/Widgets/SOpenColorIODisplay.h
+++ b/Engine/Plugins/Compositing/OpenColorIO/Source/OpenColorIOEditor/Private/Widgets/SOpenColorIODisplay.h
@@ -6,12 +6,12 @@
 
 #include "CoreMinimal.h"
 #include "OpenColorIOColorSpace.h"
+#include "SOpenColorIOColorSpacePicker.h"
 #include "Templates/SharedPointer.h"
 #include "Widgets/DeclarativeSyntaxSupport.h"
 
 
 class FViewport;
-class SOpenColorIOColorSpacePicker;
 class SWidget;
 class UToolMenu;
 struct FAssetData;
@@ -40,18 +40,24 @@ protected:
 	void ToggleEnableDisplay();
 	bool CanEnableDisplayConfiguration();
 	bool GetDisplayConfigurationState();
-	void OnSourceColorSpaceChanged(const FOpenColorIOColorSpace& NewColorSpace, const FOpenColorIODisplayView&);
+	void OnSourceColorSpaceChanged(const FOpenColorIOColorSpace& NewColorSpace, const FOpenColorIODisplayView& NewDisplayView);
 	void OnDestinationColorSpaceChanged(const FOpenColorIOColorSpace& NewColorSpace, const FOpenColorIODisplayView& NewDisplayView);
 
 protected:
+	/** Read configuration settings into local selection objects.*/
+	void ApplyConfigurationToSelection();
+	/** Apply transform selections onto the configuration object.*/
+	void ApplySelectionToConfiguration();
+
 	FViewport* Viewport;
 
 	/**Class picker widget to be able to close back the menu once selection has been made */
 	TSharedPtr<SWidget> ClassPicker;
 
 	/** ColorSpace pickers reference to update them when config asset is changed */
-	TSharedPtr<SOpenColorIOColorSpacePicker> TransformSourcePicker;
-	TSharedPtr<SOpenColorIOColorSpacePicker> TransformDestinationPicker;
+	TStaticArray<TSharedPtr<SOpenColorIOColorSpacePicker>, 2> TransformPicker;
+	/** Intermediate selection objects. This is done to allow inverted display-view selections (when enabled in settings). */
+	TStaticArray<FOpenColorIOPickerSelection, 2> TransformSelection;
 
 	/** Current configuration */
 	FOpenColorIODisplayConfiguration Configuration;
