OCIO: Improve scripting API consistency for pipeline workflows.

diff --git a/Engine/Plugins/Compositing/OpenColorIO/Source/OpenColorIO/Public/OpenColorIOColorSpace.h b/Engine/Plugins/Compositing/OpenColorIO/Source/OpenColorIO/Public/OpenColorIOColorSpace.h
--- a/Engine/Plugins/Compositing/OpenColorIO/Source/OpenColorIO/Public/OpenColorIOColorSpace.h
+++ b/Engine/Plugins/Compositing/OpenColorIO/Source/OpenColorIO/Public/OpenColorIOColorSpace.h
@@ -28,18 +28,18 @@ public:
 	FOpenColorIOColorSpace(const FString& InColorSpaceName, int32 InColorSpaceIndex, const FString& InFamilyName);
 
 	/** The ColorSpace name. */
-	UPROPERTY(VisibleAnywhere, Category=ColorSpace)
+	UPROPERTY(EditAnywhere, BlueprintReadWrite, Category=ColorSpace)
 	FString ColorSpaceName;
 
 	/** The index of the ColorSpace in the config */
-	UPROPERTY(VisibleAnywhere, Category=ColorSpace)
+	UPROPERTY(EditAnywhere, BlueprintReadWrite, Category=ColorSpace)
 	int32 ColorSpaceIndex;
 
 	/** 
 	 * The family of this ColorSpace as specified in the configuration file. 
 	 * When you have lots of colorspaces, you can regroup them by family to facilitate browsing them. 
 	 */
-	UPROPERTY(VisibleAnywhere, Category=ColorSpace)
+	UPROPERTY(EditAnywhere, BlueprintReadWrite, Category=ColorSpace)
 	FString FamilyName;
 
 	/** Delimiter used in the OpenColorIO library to make family hierarchies */
@@ -95,10 +95,10 @@ public:
 	 */
 	FOpenColorIODisplayView(FStringView InDisplayName, FStringView InViewName);
 
-	UPROPERTY(VisibleAnywhere, Category = ColorSpace)
+	UPROPERTY(EditAnywhere, BlueprintReadWrite, Category = ColorSpace)
 	FString Display;
 
-	UPROPERTY(VisibleAnywhere, Category = ColorSpace)
+	UPROPERTY(EditAnywhere, BlueprintReadWrite, Category = ColorSpace)
 	FString View;
 
 	FString ToString() const;
diff --git a/Engine/Plugins/Compositing/OpenColorIO/Source/OpenColorIO/Public/OpenColorIOConfiguration.h b/Engine/Plugins/Compositing/OpenColorIO/Source/OpenColorIO/Public/OpenColorIOConfiguration.h
--- a/Engine/Plugins/Compositing/OpenColorIO/Source/OpenColorIO/Public/OpenColorIOConfiguration.h
+++ b/Engine/Plugins/Compositing/OpenColorIO/Source/OpenColorIO/Public/OpenColorIOConfiguration.h
@@ -45,6 +45,7 @@ public:
 	bool Validate() const;
 
 	/** This forces to reload colorspaces and corresponding shaders if those are not loaded already. */
+	UFUNCTION(BlueprintCallable, Category = "OpenColorIO")
 	void ReloadExistingColorspaces();
 
 	/*
diff --git a/Engine/Plugins/Compositing/OpenColorIO/Source/OpenColorIOEditor/Private/OpenColorIOConfigurationCustomization.cpp b/Engine/Plugins/Compositing/OpenColorIO/Source/OpenColorIOEditor/Private/OpenColorIOConfigurationCustomization.cpp
--- a/Engine/Plugins/Compositing/OpenColorIO/Source/OpenColorIOEditor/Private/OpenColorIOConfigurationCustomization.cpp
+++ b/Engine/Plugins/Compositing/OpenColorIO/Source/OpenColorIOEditor/Private/OpenColorIOConfigurationCustomization.cpp
@@ -2,6 +2,7 @@
 
 #include "OpenColorIOConfigurationCustomization.h"
 
+#include "OpenColorIOColorSpaceCustomization.h"
 #include "OpenColorIOConfiguration.h"
 #include "Containers/StringConv.h"
 #include "DetailWidgetRow.h"
@@ -50,6 +51,14 @@ void FOpenColorIOConfigurationCustomization::CustomizeDetails(IDetailLayoutBuild
 				return FReply::Handled();
 			})
 		];
+
+	DetailBuilder.RegisterInstancedCustomPropertyTypeLayout(
+			FOpenColorIOColorSpace::StaticStruct()->GetFName(),
+			FOnGetPropertyTypeCustomizationInstance::CreateLambda([] { return MakeShared<FOpenColorIOColorSpaceCustomization>(); }));
+
+	DetailBuilder.RegisterInstancedCustomPropertyTypeLayout(
+		FOpenColorIODisplayView::StaticStruct()->GetFName(),
+		FOnGetPropertyTypeCustomizationInstance::CreateLambda([] { return MakeShared<FOpenColorIODisplayViewCustomization>(); }));
 }
 
 #undef LOCTEXT_NAMESPACE
diff --git a/Engine/Plugins/Compositing/OpenColorIO/Source/OpenColorIOEditor/Private/OpenColorIOEditorModule.cpp b/Engine/Plugins/Compositing/OpenColorIO/Source/OpenColorIOEditor/Private/OpenColorIOEditorModule.cpp
--- a/Engine/Plugins/Compositing/OpenColorIO/Source/OpenColorIOEditor/Private/OpenColorIOEditorModule.cpp
+++ b/Engine/Plugins/Compositing/OpenColorIO/Source/OpenColorIOEditor/Private/OpenColorIOEditorModule.cpp
@@ -13,7 +13,6 @@
 #include "OpenColorIOConfiguration.h"
 #include "OpenColorIOEditorSettings.h"
 #include "OpenColorIOColorSpaceConversionCustomization.h"
-#include "OpenColorIOColorSpaceCustomization.h"
 #include "OpenColorIOConfigurationCustomization.h"
 #include "OpenColorIOColorTransform.h"
 #include "PropertyEditorModule.h"
@@ -78,8 +77,6 @@ void FOpenColorIOEditorModule::RegisterCustomizations()
 {
 	FPropertyEditorModule& PropertyModule = FModuleManager::LoadModuleChecked<FPropertyEditorModule>("PropertyEditor");
 	PropertyModule.RegisterCustomPropertyTypeLayout(FOpenColorIOColorConversionSettings::StaticStruct()->GetFName(), FOnGetPropertyTypeCustomizationInstance::CreateStatic(&FOpenColorIOColorConversionSettingsCustomization::MakeInstance));
-	PropertyModule.RegisterCustomPropertyTypeLayout(FOpenColorIOColorSpace::StaticStruct()->GetFName(), FOnGetPropertyTypeCustomizationInstance::CreateStatic(&FOpenColorIOColorSpaceCustomization::MakeInstance));
-	PropertyModule.RegisterCustomPropertyTypeLayout(FOpenColorIODisplayView::StaticStruct()->GetFName(), FOnGetPropertyTypeCustomizationInstance::CreateStatic(&FOpenColorIODisplayViewCustomization::MakeInstance));
 	PropertyModule.RegisterCustomClassLayout(UOpenColorIOConfiguration::StaticClass()->GetFName(), FOnGetDetailCustomizationInstance::CreateStatic(&FOpenColorIOConfigurationCustomization::MakeInstance));
 }
 
