OCIO: Adding support for built-in configs URI.

The current options are:
* ocio://default (CG config equivalent)
* ocio://cg-config-v1.0.0_aces-v1.3_ocio-v2.1
* ocio://studio-config-v1.0.0_aces-v1.3_ocio-v2.1

diff --git a/Engine/Plugins/Compositing/OpenColorIO/Source/OpenColorIO/Private/OpenColorIOColorTransform.cpp b/Engine/Plugins/Compositing/OpenColorIO/Source/OpenColorIO/Private/OpenColorIOColorTransform.cpp
--- a/Engine/Plugins/Compositing/OpenColorIO/Source/OpenColorIO/Private/OpenColorIOColorTransform.cpp
+++ b/Engine/Plugins/Compositing/OpenColorIO/Source/OpenColorIO/Private/OpenColorIOColorTransform.cpp
@@ -247,7 +247,7 @@ void UOpenColorIOColorTransform::CacheResourceTextures()
 	if (Textures.IsEmpty())
 	{
 #if WITH_EDITOR && WITH_OCIO
-		OCIO_NAMESPACE::ConstConfigRcPtr CurrentConfig = ConfigurationOwner->GetLoadedConfigurationFile();
+		OCIO_NAMESPACE::ConstConfigRcPtr CurrentConfig = ConfigurationOwner->GetLoadedConfiguration();
 		if (CurrentConfig)
 		{
 #if !PLATFORM_EXCEPTIONS_DISABLED
@@ -570,7 +570,7 @@ bool UOpenColorIOColorTransform::UpdateShaderInfo(FString& OutShaderCodeHash, FS
 {
 #if WITH_EDITOR
 #if WITH_OCIO
-	OCIO_NAMESPACE::ConstConfigRcPtr CurrentConfig = ConfigurationOwner->GetLoadedConfigurationFile();
+	OCIO_NAMESPACE::ConstConfigRcPtr CurrentConfig = ConfigurationOwner->GetLoadedConfiguration();
 	if (CurrentConfig)
 	{
 #if !PLATFORM_EXCEPTIONS_DISABLED
diff --git a/Engine/Plugins/Compositing/OpenColorIO/Source/OpenColorIO/Private/OpenColorIOConfiguration.cpp b/Engine/Plugins/Compositing/OpenColorIO/Source/OpenColorIO/Private/OpenColorIOConfiguration.cpp
--- a/Engine/Plugins/Compositing/OpenColorIO/Source/OpenColorIO/Private/OpenColorIOConfiguration.cpp
+++ b/Engine/Plugins/Compositing/OpenColorIO/Source/OpenColorIO/Private/OpenColorIOConfiguration.cpp
@@ -50,6 +50,7 @@ namespace OCIODirectoryWatcher
 UOpenColorIOConfiguration::UOpenColorIOConfiguration(const FObjectInitializer& ObjectInitializer)
 	: Super(ObjectInitializer)
 {
+	
 }
 
 void UOpenColorIOConfiguration::BeginDestroy()
@@ -171,7 +172,7 @@ void UOpenColorIOConfiguration::ReloadExistingColorspaces()
 	DesiredColorSpaces.Reset();
 	DesiredDisplayViews.Reset();
 	CleanupTransforms();
-	LoadConfigurationFile();
+	LoadConfiguration();
 
 	if (!LoadedConfig)
 	{
@@ -421,6 +422,20 @@ void UOpenColorIOConfiguration::CleanupTransforms()
 		});
 }
 
+
+void UOpenColorIOConfiguration::PostInitProperties()
+{
+	Super::PostInitProperties();
+
+	if (!HasAnyFlags(RF_NeedLoad | RF_ClassDefaultObject))
+	{
+		ConfigurationFile.FilePath = TEXT("ocio://default");
+
+		// Ensure the default built-in configuration is loaded.
+		LoadConfiguration();
+	}
+}
+
 void UOpenColorIOConfiguration::PostLoad()
 {
 	Super::PostLoad();
@@ -472,7 +487,7 @@ void UOpenColorIOConfiguration::PostEditChangeProperty(FPropertyChangedEvent& Pr
 
 	if (PropertyChangedEvent.MemberProperty && PropertyChangedEvent.MemberProperty->GetFName() == GET_MEMBER_NAME_CHECKED(UOpenColorIOConfiguration, ConfigurationFile))
 	{
-		LoadConfigurationFile();
+		LoadConfiguration();
 	}
 	else if (PropertyChangedEvent.GetPropertyName() == GET_MEMBER_NAME_CHECKED(UOpenColorIOConfiguration, DesiredColorSpaces))
 	{
@@ -534,7 +549,7 @@ void UOpenColorIOConfiguration::PostEditChangeProperty(FPropertyChangedEvent& Pr
 }
 #endif //WITH_EDITOR
 
-void UOpenColorIOConfiguration::LoadConfigurationFile()
+void UOpenColorIOConfiguration::LoadConfiguration()
 {
 #if WITH_EDITOR && WITH_OCIO
 	if (!ConfigurationFile.FilePath.IsEmpty())
@@ -543,41 +558,49 @@ void UOpenColorIOConfiguration::LoadConfigurationFile()
 		try
 #endif
 		{
-			LoadedConfig.reset();
-
-			FString FullPath;
 			FString ConfigurationFilePath = ConfigurationFile.FilePath;
-			if (ConfigurationFilePath.Contains(TEXT("{Engine}")))
-			{
-				ConfigurationFilePath = FPaths::ConvertRelativePathToFull(ConfigurationFilePath.Replace(TEXT("{Engine}"), *FPaths::EngineDir()));
-			}    
-
-			if (!FPaths::IsRelative(ConfigurationFilePath))
+			if (ConfigurationFilePath.StartsWith(TEXT("ocio://")))
 			{
-				FullPath = ConfigurationFilePath;
+				LoadedConfig = OCIO_NAMESPACE::Config::CreateFromFile(StringCast<ANSICHAR>(*ConfigurationFilePath).Get());
+				UE_LOG(LogOpenColorIO, Verbose, TEXT("Loaded built-in OCIO configuration file %s"), *ConfigurationFilePath);
 			}
 			else
 			{
-				const FString AbsoluteGameDir = FPaths::ConvertRelativePathToFull(FPaths::ProjectDir());
-				FullPath = FPaths::ConvertRelativePathToFull(FPaths::Combine(AbsoluteGameDir, ConfigurationFilePath));
-			}
+				LoadedConfig.reset();
 
-			OCIO_NAMESPACE::ConstConfigRcPtr NewConfig = OCIO_NAMESPACE::Config::CreateFromFile(StringCast<ANSICHAR>(*FullPath).Get());
-			if (NewConfig)
-			{
-				UE_LOG(LogOpenColorIO, Verbose, TEXT("Loaded OCIO configuration file %s"), *FullPath);
-				LoadedConfig = NewConfig;
-				StartDirectoryWatch(FullPath);
-			}
-			else
-			{
-				UE_LOG(LogOpenColorIO, Error, TEXT("Could not load OCIO configuration file %s. Verify that the path is good or that the file is valid."), *ConfigurationFile.FilePath);
+				FString FullPath;
+				if (ConfigurationFilePath.Contains(TEXT("{Engine}")))
+				{
+					ConfigurationFilePath = FPaths::ConvertRelativePathToFull(ConfigurationFilePath.Replace(TEXT("{Engine}"), *FPaths::EngineDir()));
+				}
+
+				if (!FPaths::IsRelative(ConfigurationFilePath))
+				{
+					FullPath = ConfigurationFilePath;
+				}
+				else
+				{
+					const FString AbsoluteGameDir = FPaths::ConvertRelativePathToFull(FPaths::ProjectDir());
+					FullPath = FPaths::ConvertRelativePathToFull(FPaths::Combine(AbsoluteGameDir, ConfigurationFilePath));
+				}
+
+				OCIO_NAMESPACE::ConstConfigRcPtr NewConfig = OCIO_NAMESPACE::Config::CreateFromFile(StringCast<ANSICHAR>(*FullPath).Get());
+				if (NewConfig)
+				{
+					UE_LOG(LogOpenColorIO, Verbose, TEXT("Loaded OCIO configuration file %s"), *FullPath);
+					LoadedConfig = NewConfig;
+					StartDirectoryWatch(FullPath);
+				}
+				else
+				{
+					UE_LOG(LogOpenColorIO, Error, TEXT("Could not load OCIO configuration file %s. Verify that the path is good or that the file is valid."), *ConfigurationFile.FilePath);
+				}
 			}
 		}
 #if !PLATFORM_EXCEPTIONS_DISABLED
 		catch (OCIO_NAMESPACE::Exception& exception)
 		{
-			UE_LOG(LogOpenColorIO, Error, TEXT("Could not load OCIO configuration file %s. Error message: %s."), *ConfigurationFile.FilePath, StringCast<TCHAR>(exception.what()).Get());
+			UE_LOG(LogOpenColorIO, Error, TEXT("Could not create OCIO configuration file for %s. Error message: %s."), *ConfigurationFile.FilePath, StringCast<TCHAR>(exception.what()).Get());
 		}
 #endif
 	}
diff --git a/Engine/Plugins/Compositing/OpenColorIO/Source/OpenColorIO/Public/OpenColorIOConfiguration.h b/Engine/Plugins/Compositing/OpenColorIO/Source/OpenColorIO/Public/OpenColorIOConfiguration.h
--- a/Engine/Plugins/Compositing/OpenColorIO/Source/OpenColorIO/Public/OpenColorIOConfiguration.h
+++ b/Engine/Plugins/Compositing/OpenColorIO/Source/OpenColorIO/Public/OpenColorIOConfiguration.h
@@ -55,7 +55,10 @@ public:
 	void ConfigPathChangedEvent(const TArray<FFileChangeData>& InFileChanges, const FString InFileMountPath);
 
 #if WITH_EDITORONLY_DATA && WITH_OCIO
-	OCIO_NAMESPACE::ConstConfigRcPtr GetLoadedConfigurationFile() const { return LoadedConfig; }
+	OCIO_NAMESPACE::ConstConfigRcPtr GetLoadedConfiguration() const { return LoadedConfig; }
+
+	UE_DEPRECATED(5.2, "GetLoadedConfigurationFile is deprecated, please use GetLoadedConfiguration instead.")
+	OCIO_NAMESPACE::ConstConfigRcPtr GetLoadedConfigurationFile() const { return GetLoadedConfiguration(); }
 #endif
 
 protected:
@@ -72,6 +75,7 @@ protected:
 public:
 
 	//~ Begin UObject interface
+	virtual void PostInitProperties() override;
 	virtual void PostLoad() override;
 #if WITH_EDITORONLY_DATA
 	static void DeclareConstructClasses(TArray<FTopLevelAssetPath>& OutConstructClasses, const UClass* SpecificSubclass);
@@ -86,7 +90,7 @@ public:
 	//~ End UObject interface
 
 private:
-	void LoadConfigurationFile();
+	void LoadConfiguration();
 
 	/** This method resets the status of Notification dialog and reacts depending on user's choice. */
 	void OnToastCallback(bool bInReloadColorspaces);
