OCIO: Prevent thumbnail hover from loading the entire asset.

#rnx

diff --git a/Engine/Plugins/Compositing/OpenColorIO/Source/OpenColorIO/Private/OpenColorIOConfiguration.cpp b/Engine/Plugins/Compositing/OpenColorIO/Source/OpenColorIO/Private/OpenColorIOConfiguration.cpp
--- a/Engine/Plugins/Compositing/OpenColorIO/Source/OpenColorIO/Private/OpenColorIOConfiguration.cpp
+++ b/Engine/Plugins/Compositing/OpenColorIO/Source/OpenColorIO/Private/OpenColorIOConfiguration.cpp
@@ -447,6 +447,27 @@ void UOpenColorIOConfiguration::PostLoad()
 	}
 }
 
+void UOpenColorIOConfiguration::GetAssetRegistryTags(TArray<FAssetRegistryTag>& OutTags) const
+{
+	Super::GetAssetRegistryTags(OutTags);
+
+	FString Description;
+	if (ConfigurationFile.FilePath.IsEmpty())
+	{
+		Description = TEXT("No configuration selected.");
+	}
+	else if (!Validate())
+	{
+		Description = TEXT("Warning: Configuration is invalid. Verify the selected configuration file.");
+	}
+	else
+	{
+		Description = TEXT("Configuration: ") + ConfigurationFile.FilePath;
+	}
+
+	OutTags.Add(FAssetRegistryTag(TEXT("ConfigurationFile"), Description, FAssetRegistryTag::TT_Hidden));
+}
+
 #if WITH_EDITORONLY_DATA
 void UOpenColorIOConfiguration::DeclareConstructClasses(TArray<FTopLevelAssetPath>& OutConstructClasses, const UClass* SpecificSubclass)
 {
diff --git a/Engine/Plugins/Compositing/OpenColorIO/Source/OpenColorIO/Public/OpenColorIOConfiguration.h b/Engine/Plugins/Compositing/OpenColorIO/Source/OpenColorIO/Public/OpenColorIOConfiguration.h
--- a/Engine/Plugins/Compositing/OpenColorIO/Source/OpenColorIO/Public/OpenColorIOConfiguration.h
+++ b/Engine/Plugins/Compositing/OpenColorIO/Source/OpenColorIO/Public/OpenColorIOConfiguration.h
@@ -77,6 +77,7 @@ public:
 	//~ Begin UObject interface
 	virtual void PostInitProperties() override;
 	virtual void PostLoad() override;
+	virtual void GetAssetRegistryTags(TArray<FAssetRegistryTag>& OutTags) const override;
 #if WITH_EDITORONLY_DATA
 	static void DeclareConstructClasses(TArray<FTopLevelAssetPath>& OutConstructClasses, const UClass* SpecificSubclass);
 #endif
diff --git a/Engine/Plugins/Compositing/OpenColorIO/Source/OpenColorIOEditor/Private/AssetTypeActions_OpenColorIOConfiguration.cpp b/Engine/Plugins/Compositing/OpenColorIO/Source/OpenColorIOEditor/Private/AssetTypeActions_OpenColorIOConfiguration.cpp
--- a/Engine/Plugins/Compositing/OpenColorIO/Source/OpenColorIOEditor/Private/AssetTypeActions_OpenColorIOConfiguration.cpp
+++ b/Engine/Plugins/Compositing/OpenColorIO/Source/OpenColorIOEditor/Private/AssetTypeActions_OpenColorIOConfiguration.cpp
@@ -12,30 +12,7 @@
 
 FText FAssetTypeActions_OpenColorIOConfiguration::GetAssetDescription(const struct FAssetData& AssetData) const
 {
-	UOpenColorIOConfiguration* Asset = Cast<UOpenColorIOConfiguration>(AssetData.GetAsset());
-
-	FString Description;
-	if (Asset != nullptr)
-	{
-		const FString AssociatedFile = Asset->ConfigurationFile.FilePath;
-
-		if (AssociatedFile.IsEmpty())
-		{
-			Description = TEXT("No configuration file selected.");
-		}
-		else if (!Asset->Validate())
-		{
-			Description = TEXT("Warning: Configuration asset is invalid. Verify the selected configuration file.");
-		}
-		else
-		{
-			Description = TEXT("Configuration file: ") + AssociatedFile;
-		}
-
-		
-	}
-
-	return FText::Format(LOCTEXT("AssetTypeActions_OpenColorIOConfiguration_Description", "{0}"), FText::FromString(MoveTemp(Description)));
+	return FText::FromString(AssetData.GetTagValueRef<FString>(GET_MEMBER_NAME_CHECKED(UOpenColorIOConfiguration, ConfigurationFile)));
 }
 
 
