OCIO: Fix property update bug, introduced by the bidirectional display-view selection change.

Properties need to be batched updated as a single transaction, otherwise the modification of the first property would cause ApplyConfigurationToSelection() to be called immediately, invalidating the change.

#rnx

diff --git a/Engine/Plugins/Compositing/OpenColorIO/Source/OpenColorIOEditor/Private/OpenColorIOColorSpaceConversionCustomization.cpp b/Engine/Plugins/Compositing/OpenColorIO/Source/OpenColorIOEditor/Private/OpenColorIOColorSpaceConversionCustomization.cpp
--- a/Engine/Plugins/Compositing/OpenColorIO/Source/OpenColorIOEditor/Private/OpenColorIOColorSpaceConversionCustomization.cpp
+++ b/Engine/Plugins/Compositing/OpenColorIO/Source/OpenColorIOEditor/Private/OpenColorIOColorSpaceConversionCustomization.cpp
@@ -12,6 +12,7 @@
 #include "OpenColorIOConfiguration.h"
 #include "OpenColorIOSettings.h"
 #include "PropertyHandle.h"
+#include "ScopedTransaction.h"
 #include "Widgets/SOpenColorIOColorSpacePicker.h"
 #include "Widgets/Text/STextBlock.h"
 
@@ -60,24 +61,6 @@ void FOpenColorIOColorConversionSettingsCustomization::CustomizeHeader(TSharedRe
 	}
 }
 
-namespace
-{
-	void UpdateColorSettingsStructProperty(const TSharedPtr<IPropertyHandle>& PropertyHandle, const void* Value, const void* Defaults)
-	{
-		if (PropertyHandle.IsValid())
-		{
-			if (FStructProperty* StructProperty = CastField<FStructProperty>(PropertyHandle->GetProperty()))
-			{
-				FString TextValue;
-				StructProperty->Struct->ExportText(TextValue, Value, Defaults, nullptr, EPropertyPortFlags::PPF_None, nullptr);
-
-				// Note: using set value is preferable for built-in change propagation
-				PropertyHandle->SetValueFromFormattedString(TextValue);
-			}
-		}
-	}
-}
-
 void FOpenColorIOColorConversionSettingsCustomization::CustomizeChildren(TSharedRef<IPropertyHandle> InStructPropertyHandle, IDetailChildrenBuilder& StructBuilder, IPropertyTypeCustomizationUtils& StructCustomizationUtils)
 {
 	if (ColorSpaceConversion == nullptr)
@@ -224,24 +207,46 @@ void FOpenColorIOColorConversionSettingsCustomization::ApplyConfigurationToSelec
 
 void FOpenColorIOColorConversionSettingsCustomization::ApplySelectionToConfiguration()
 {
-	check(SourceColorSpaceProperty.IsValid() && DestinationColorSpaceProperty.IsValid() && DestinationDisplayViewProperty.IsValid() && DisplayViewDirectionProperty.IsValid());
+	TArray<TSharedPtr<IPropertyHandle>> Properties = { SourceColorSpaceProperty, DestinationColorSpaceProperty, DestinationDisplayViewProperty, DisplayViewDirectionProperty };
 
-	FOpenColorIOColorConversionSettings& ConversionSettings = *ColorSpaceConversion;
+	const FScopedTransaction Transaction(LOCTEXT("OCIOConfigurationSelectionUpdate", "OCIO Configuration Selection Update"));
+
+	// Notify outer object(s). Only called for one property since it is shared.
+	TArray<UObject*> Objects;
+	Properties[0]->GetOuterObjects(Objects);
+	for (UObject* OuterObject : Objects)
+	{
+		OuterObject->Modify();
+	}
+
+	for (const TSharedPtr<IPropertyHandle>& Property : Properties)
+	{
+		check(Property.IsValid());
+		Property->NotifyPreChange();
+	}
 
 	if (TransformSelection[OCIO_Src].DisplayView.IsValid())
 	{
-		FOpenColorIOColorSpace EmptyCS = {};
-		UpdateColorSettingsStructProperty(SourceColorSpaceProperty, &TransformSelection[OCIO_Dst].ColorSpace, &ColorSpaceConversion->SourceColorSpace);
-		UpdateColorSettingsStructProperty(DestinationColorSpaceProperty, &EmptyCS, &ColorSpaceConversion->DestinationColorSpace);
-		UpdateColorSettingsStructProperty(DestinationDisplayViewProperty, &TransformSelection[OCIO_Src].DisplayView, &ColorSpaceConversion->DestinationDisplayView);
-		DisplayViewDirectionProperty->SetValue((uint8)EOpenColorIOViewTransformDirection::Inverse);
+		ColorSpaceConversion->SourceColorSpace = TransformSelection[OCIO_Dst].ColorSpace;
+		ColorSpaceConversion->DestinationColorSpace.Reset();
+		ColorSpaceConversion->DestinationDisplayView = TransformSelection[OCIO_Src].DisplayView;
+		ColorSpaceConversion->DisplayViewDirection = EOpenColorIOViewTransformDirection::Inverse;
 	}
 	else
 	{
-		UpdateColorSettingsStructProperty(SourceColorSpaceProperty, &TransformSelection[OCIO_Src].ColorSpace, &ColorSpaceConversion->SourceColorSpace);
-		UpdateColorSettingsStructProperty(DestinationColorSpaceProperty, &TransformSelection[OCIO_Dst].ColorSpace, &ColorSpaceConversion->DestinationColorSpace);
-		UpdateColorSettingsStructProperty(DestinationDisplayViewProperty, &TransformSelection[OCIO_Dst].DisplayView, &ColorSpaceConversion->DestinationDisplayView);
-		DisplayViewDirectionProperty->SetValue((uint8)EOpenColorIOViewTransformDirection::Forward);
+		ColorSpaceConversion->SourceColorSpace = TransformSelection[OCIO_Src].ColorSpace;
+		ColorSpaceConversion->DestinationColorSpace = TransformSelection[OCIO_Dst].ColorSpace;
+		ColorSpaceConversion->DestinationDisplayView = TransformSelection[OCIO_Dst].DisplayView;
+		ColorSpaceConversion->DisplayViewDirection = EOpenColorIOViewTransformDirection::Forward;
+	}
+
+	for (const TSharedPtr<IPropertyHandle>& Property : Properties)
+	{
+		Property->NotifyPostChange(EPropertyChangeType::ValueSet);
+	}
+	for (const TSharedPtr<IPropertyHandle>& Property : Properties)
+	{
+		Property->NotifyFinishedChangingProperties();
 	}
 }
 
