OCIO: Added SVE error fallback render pass.

Fixes previously all black render on nDisplay on invalid transform with a scene passthrough and a visible error message. Removed duplicate error messages.

diff --git a/Engine/Plugins/Compositing/OpenColorIO/Shaders/Private/OpenColorIOErrorShader.usf b/Engine/Plugins/Compositing/OpenColorIO/Shaders/Private/OpenColorIOErrorShader.usf
new file mode 100644
--- /dev/null
+++ b/Engine/Plugins/Compositing/OpenColorIO/Shaders/Private/OpenColorIOErrorShader.usf
@@ -0,0 +1,28 @@
+// Copyright Epic Games, Inc. All Rights Reserved.
+
+/*=============================================================================
+ OpenColorIOErrorShader.usf: Fallback OCIO shader when there is an error with the transform.
+ =============================================================================*/
+
+#include "/Engine/Public/Platform.ush"
+#include "/Engine/Private/MiniFontCommon.ush"
+
+Texture2D		InputTexture;
+SamplerState	InputTextureSampler;
+
+void MainPS(in noperspective float4 InUVAndScreenPos : TEXCOORD0,
+	float4 SvPosition : SV_POSITION,
+	out float4 OutColor : SV_Target0)
+{
+	const int2 PixelPos = (int2)SvPosition.xy % int2(256, 256); // Tile message every 256px
+	const float3 TextColor = float3(1, 0.75, 0.25);
+	int2 Cursor = int2(128, 128);
+
+	OutColor = InputTexture.Sample(InputTextureSampler, InUVAndScreenPos.xy);
+
+#define P(x)  PrintCharacter(PixelPos, OutColor.xyz, TextColor, Cursor, x);
+
+	P(_O_)P(_C_)P(_I_)P(_O_)
+	P(_SPC_)
+	P(_E_)P(_R_)P(_R_)P(_O_)P(_R_)
+}
\ No newline at end of file
diff --git a/Engine/Plugins/Compositing/OpenColorIO/Source/OpenColorIO/Private/OpenColorIOColorTransform.cpp b/Engine/Plugins/Compositing/OpenColorIO/Source/OpenColorIO/Private/OpenColorIOColorTransform.cpp
--- a/Engine/Plugins/Compositing/OpenColorIO/Source/OpenColorIO/Private/OpenColorIOColorTransform.cpp
+++ b/Engine/Plugins/Compositing/OpenColorIO/Source/OpenColorIO/Private/OpenColorIOColorTransform.cpp
@@ -432,6 +432,7 @@ FOpenColorIOTransformResource* UOpenColorIOColorTransform::AllocateResource()
 bool UOpenColorIOColorTransform::GetRenderResources(ERHIFeatureLevel::Type InFeatureLevel, FOpenColorIOTransformResource*& OutShaderResource, TSortedMap<int32, FTextureResource*>& OutTextureResources)
 {
 	OutShaderResource = ColorTransformResources[InFeatureLevel];
+	
 	if (OutShaderResource)
 	{
 		OutTextureResources.Reserve(Textures.Num());
@@ -440,14 +441,9 @@ bool UOpenColorIOColorTransform::GetRenderResources(ERHIFeatureLevel::Type InFea
 		{
 			OutTextureResources.Add(Pair.Key, Pair.Value->GetResource());
 		}
-
-		return true;
-	}
-	else
-	{
-		UE_LOG(LogOpenColorIO, Warning, TEXT("Shader resource was invalid for color transform %s. Were there errors during loading?"), *GetTransformFriendlyName());
-		return false;
 	}
+	
+	return OutShaderResource != nullptr;
 }
 
 bool UOpenColorIOColorTransform::GetShaderAndLUTResouces(ERHIFeatureLevel::Type InFeatureLevel, FOpenColorIOTransformResource*& OutShaderResource, FTextureResource*& OutLUT3dResource)
diff --git a/Engine/Plugins/Compositing/OpenColorIO/Source/OpenColorIO/Private/OpenColorIOConfiguration.cpp b/Engine/Plugins/Compositing/OpenColorIO/Source/OpenColorIO/Private/OpenColorIOConfiguration.cpp
--- a/Engine/Plugins/Compositing/OpenColorIO/Source/OpenColorIO/Private/OpenColorIOConfiguration.cpp
+++ b/Engine/Plugins/Compositing/OpenColorIO/Source/OpenColorIO/Private/OpenColorIOConfiguration.cpp
@@ -77,7 +77,6 @@ bool UOpenColorIOConfiguration::GetRenderResources(ERHIFeatureLevel::Type InFeat
 
 	if (TransformPtr == nullptr)
 	{
-		UE_LOG(LogOpenColorIO, Warning, TEXT("Color transform [%s] was not found."), *InSettings.ToString());
 		return false;
 	}
 
diff --git a/Engine/Plugins/Compositing/OpenColorIO/Source/OpenColorIO/Private/OpenColorIODisplayExtension.cpp b/Engine/Plugins/Compositing/OpenColorIO/Source/OpenColorIO/Private/OpenColorIODisplayExtension.cpp
--- a/Engine/Plugins/Compositing/OpenColorIO/Source/OpenColorIO/Private/OpenColorIODisplayExtension.cpp
+++ b/Engine/Plugins/Compositing/OpenColorIO/Source/OpenColorIO/Private/OpenColorIODisplayExtension.cpp
@@ -15,6 +15,8 @@
 #include "PipelineStateCache.h"
 #include "RenderTargetPool.h"
 #include "Shader.h"
+#include "ShaderParameters.h"
+#include "ShaderParameterStruct.h"
 
 #include "RHI.h"
 #include "SceneView.h"
@@ -25,8 +27,35 @@
 // for FPostProcessMaterialInputs
 #include "PostProcess/PostProcessMaterial.h"
 
+
+BEGIN_SHADER_PARAMETER_STRUCT(FOpenColorIOErrorShaderParameters, )
+	SHADER_PARAMETER_RDG_TEXTURE(Texture2D, InputTexture)
+	SHADER_PARAMETER_SAMPLER(SamplerState, InputTextureSampler)
+	SHADER_PARAMETER_TEXTURE(Texture2D, MiniFontTexture)
+	RENDER_TARGET_BINDING_SLOTS()
+END_SHADER_PARAMETER_STRUCT()
+
+class FOpenColorIOErrorPassPS : public FGlobalShader
+{
+public:
+	DECLARE_GLOBAL_SHADER(FOpenColorIOErrorPassPS);
+	SHADER_USE_PARAMETER_STRUCT(FOpenColorIOErrorPassPS, FGlobalShader);
+
+	using FParameters = FOpenColorIOErrorShaderParameters;
+};
+
+IMPLEMENT_GLOBAL_SHADER(FOpenColorIOErrorPassPS, "/Plugin/OpenColorIO/Private/OpenColorIOErrorShader.usf", "MainPS", SF_Pixel);
+
+
 float FOpenColorIODisplayExtension::DefaultDisplayGamma = 2.2f;
 
+namespace {
+	FRHITexture* GetSystemMiniFontTexture()
+	{
+		return GSystemTextures.AsciiTexture ? GSystemTextures.AsciiTexture->GetRHI() : GSystemTextures.WhiteDummy->GetRHI();
+	}
+};
+
 FOpenColorIODisplayExtension::FOpenColorIODisplayExtension(const FAutoRegister& AutoRegister, FViewportClient* AssociatedViewportClient)
 	: FSceneViewExtensionBase(AutoRegister)
 	, LinkedViewportClient(AssociatedViewportClient)
@@ -72,7 +101,7 @@ void FOpenColorIODisplayExtension::SetupView(FSceneViewFamily& InViewFamily, FSc
 
 		if (!bFoundTransform)
 		{
-			UE_LOG(LogOpenColorIO, Warning, TEXT("Can't apply display look - Couldn't find shader to transform from %s to %s"), *DisplayConfiguration.ColorConfiguration.SourceColorSpace.ColorSpaceName, *DisplayConfiguration.ColorConfiguration.DestinationColorSpace.ColorSpaceName);
+			UE_LOG(LogOpenColorIO, Warning, TEXT("Can't apply display look - Couldn't find shader to transform: %s"), *DisplayConfiguration.ColorConfiguration.ToString());
 		}
 		else
 		{
@@ -118,27 +147,13 @@ void FOpenColorIODisplayExtension::SubscribeToPostProcessingPass(EPostProcessing
 
 FScreenPassTexture FOpenColorIODisplayExtension::PostProcessPassAfterTonemap_RenderThread(FRDGBuilder& GraphBuilder, const FSceneView& View, const FPostProcessMaterialInputs& InOutInputs)
 {
-	const FSceneViewFamily& ViewFamily = *View.Family;
 	const FScreenPassTexture& SceneColor = InOutInputs.GetInput(EPostProcessMaterialInput::SceneColor);
 	check(SceneColor.IsValid());
-
-	//If the shader resource could not be found, skip this frame. The LUT isn't required
-	if (CachedResourcesRenderThread.ShaderResource == nullptr)
-	{
-		return SceneColor;
-	}
-
-	if (!EnumHasAnyFlags(SceneColor.Texture->Desc.Flags, TexCreate_ShaderResource))
-	{
-		return SceneColor;
-	}
-
 	checkSlow(View.bIsViewInfo);
 	const FViewInfo& ViewInfo = static_cast<const FViewInfo&>(View);
-	
 	FScreenPassRenderTarget Output = InOutInputs.OverrideOutput;
 
-	// If the override output is provided it means that this is the last pass in post processing.
+	// If the override output is provided, it means that this is the last pass in post processing.
 	if (!Output.IsValid())
 	{
 		Output = FScreenPassRenderTarget::CreateFromInput(GraphBuilder, SceneColor, ViewInfo.GetOverwriteLoadAction(), TEXT("OCIORenderTarget"));
@@ -147,21 +162,37 @@ FScreenPassTexture FOpenColorIODisplayExtension::PostProcessPassAfterTonemap_Ren
 	const FScreenPassTextureViewport InputViewport(SceneColor);
 	const FScreenPassTextureViewport OutputViewport(Output);
 
-	TShaderRef<FOpenColorIOPixelShader> OCIOPixelShader = CachedResourcesRenderThread.ShaderResource->GetShader<FOpenColorIOPixelShader>();
+	if (CachedResourcesRenderThread.ShaderResource)
+	{
+		TShaderRef<FOpenColorIOPixelShader> OCIOPixelShader = CachedResourcesRenderThread.ShaderResource->GetShader<FOpenColorIOPixelShader>();
 
-	const float DisplayGamma = View.Family->RenderTarget->GetDisplayGamma();
+		const float DisplayGamma = View.Family->RenderTarget->GetDisplayGamma();
 
-	FOpenColorIOPixelShaderParameters* Parameters = GraphBuilder.AllocParameters<FOpenColorIOPixelShaderParameters>();
-	Parameters->InputTexture = SceneColor.Texture;
-	Parameters->InputTextureSampler = TStaticSamplerState<>::GetRHI();
-	OpenColorIOBindTextureResources(Parameters, CachedResourcesRenderThread.TextureResources);
+		FOpenColorIOPixelShaderParameters* Parameters = GraphBuilder.AllocParameters<FOpenColorIOPixelShaderParameters>();
+		Parameters->InputTexture = SceneColor.Texture;
+		Parameters->InputTextureSampler = TStaticSamplerState<>::GetRHI();
+		OpenColorIOBindTextureResources(Parameters, CachedResourcesRenderThread.TextureResources);
 
-	// There is a special case where post processing and tonemapper are disabled. In this case tonemapper applies a static display Inverse of Gamma which defaults to 2.2.
-	// In the case when Both PostProcessing and ToneMapper are disabled we apply gamma manually. In every other case we apply inverse gamma before applying OCIO.
-	Parameters->Gamma = (ViewFamily.EngineShowFlags.Tonemapper == 0) || (ViewFamily.EngineShowFlags.PostProcessing == 0) ? DefaultDisplayGamma : DefaultDisplayGamma / DisplayGamma;
-	Parameters->RenderTargets[0] = Output.GetRenderTargetBinding();
+		// There is a special case where post processing and tonemapper are disabled. In this case tonemapper applies a static display Inverse of Gamma which defaults to 2.2.
+		// In the case when Both PostProcessing and ToneMapper are disabled we apply gamma manually. In every other case we apply inverse gamma before applying OCIO.
+		Parameters->Gamma = (View.Family->EngineShowFlags.Tonemapper == 0) || (View.Family->EngineShowFlags.PostProcessing == 0) ? DefaultDisplayGamma : DefaultDisplayGamma / DisplayGamma;
+		Parameters->RenderTargets[0] = Output.GetRenderTargetBinding();
+
+		AddDrawScreenPass(GraphBuilder, RDG_EVENT_NAME("OCIODisplayLook"), ViewInfo, OutputViewport, InputViewport, OCIOPixelShader, Parameters);
+	}
+	else
+	{
+		// Fallback error pass, printing OCIO error message indicators across the viewport. (Helpful to quickly identify an OCIO config issue on nDisplay for example.)
+		TShaderMapRef<FOpenColorIOErrorPassPS> OCIOPixelShader(ViewInfo.ShaderMap);
+		FOpenColorIOErrorShaderParameters* Parameters = GraphBuilder.AllocParameters<FOpenColorIOErrorShaderParameters>();
+		Parameters->InputTexture = SceneColor.Texture;
+		Parameters->InputTextureSampler = TStaticSamplerState<>::GetRHI();
+		Parameters->MiniFontTexture = GetSystemMiniFontTexture();
+		Parameters->RenderTargets[0] = Output.GetRenderTargetBinding();
+
+		AddDrawScreenPass(GraphBuilder, RDG_EVENT_NAME("OCIODisplayLookError"), ViewInfo, OutputViewport, InputViewport, OCIOPixelShader, Parameters);
+	}
 
-	AddDrawScreenPass(GraphBuilder, RDG_EVENT_NAME("OCIODisplayLook"), ViewInfo, OutputViewport, InputViewport, OCIOPixelShader, Parameters);
 
 	return MoveTemp(Output);
 }
\ No newline at end of file
diff --git a/Engine/Plugins/Compositing/OpenColorIO/Source/OpenColorIO/Private/OpenColorIORendering.cpp b/Engine/Plugins/Compositing/OpenColorIO/Source/OpenColorIO/Private/OpenColorIORendering.cpp
--- a/Engine/Plugins/Compositing/OpenColorIO/Source/OpenColorIO/Private/OpenColorIORendering.cpp
+++ b/Engine/Plugins/Compositing/OpenColorIO/Source/OpenColorIO/Private/OpenColorIORendering.cpp
@@ -107,7 +107,7 @@ bool FOpenColorIORendering::ApplyColorTransform(UWorld* InWorld, const FOpenColo
 	bool bFoundTransform = InSettings.ConfigurationSource->GetRenderResources(FeatureLevel, InSettings, ShaderResource, TextureResources);
 	if (!bFoundTransform)
 	{
-		UE_LOG(LogOpenColorIO, Warning, TEXT("Can't apply color transform - Couldn't find shader to transform from %s to %s"), *InSettings.SourceColorSpace.ColorSpaceName, *InSettings.DestinationColorSpace.ColorSpaceName);
+		UE_LOG(LogOpenColorIO, Warning, TEXT("Can't apply color transform - Couldn't find shader to transform: %s"), *InSettings.ToString());
 		return false;
 	}
 
