OCIO: Cleanup logging errors to minimize spam on invalid or failed transforms, and support fallback pass in non-SVE ocio pass.

diff --git a/Engine/Plugins/Compositing/OpenColorIO/Source/OpenColorIO/Private/OpenColorIOColorTransform.cpp b/Engine/Plugins/Compositing/OpenColorIO/Source/OpenColorIO/Private/OpenColorIOColorTransform.cpp
--- a/Engine/Plugins/Compositing/OpenColorIO/Source/OpenColorIO/Private/OpenColorIOColorTransform.cpp
+++ b/Engine/Plugins/Compositing/OpenColorIO/Source/OpenColorIO/Private/OpenColorIOColorTransform.cpp
@@ -342,19 +342,19 @@ void UOpenColorIOColorTransform::CacheResourceTextures()
 				}
 				else
 				{
-					UE_LOG(LogOpenColorIO, Error, TEXT("Failed to cache 3dLUT for color transform %s. Transform processor was unusable."), *GetTransformFriendlyName());
+					UE_LOG(LogOpenColorIO, Error, TEXT("Failed to cache texture resource(s) for color transform %s. Transform processor was unusable."), *GetTransformFriendlyName());
 				}
 			}
 #if !PLATFORM_EXCEPTIONS_DISABLED
 			catch (OCIO_NAMESPACE::Exception& exception)
 			{
-				UE_LOG(LogOpenColorIO, Log, TEXT("Failed to cache 3dLUT for color transform %s. Error message: %s."), *GetTransformFriendlyName(), StringCast<TCHAR>(exception.what()).Get());
+				UE_LOG(LogOpenColorIO, Error, TEXT("Failed to cache texture resource(s) for color transform %s. Error message: %s"), *GetTransformFriendlyName(), StringCast<TCHAR>(exception.what()).Get());
 			}
 #endif
 		}
 		else
 		{
-			UE_LOG(LogOpenColorIO, Error, TEXT("Failed to cache 3dLUT for color transform %s. Configuration file was invalid."), *GetTransformFriendlyName());
+			UE_LOG(LogOpenColorIO, Error, TEXT("Failed to cache texture resource(s) for color transform %s. Configuration file [%s] was invalid."), *GetTransformFriendlyName(), *ConfigurationOwner->ConfigurationFile.FilePath);
 		}
 #endif
 	}
@@ -613,13 +613,13 @@ bool UOpenColorIOColorTransform::UpdateShaderInfo(FString& OutShaderCodeHash, FS
 #if !PLATFORM_EXCEPTIONS_DISABLED
 		catch (OCIO_NAMESPACE::Exception& exception)
 		{
-			UE_LOG(LogOpenColorIO, Log, TEXT("Failed to fetch shader info for color transform %s. Error message: %s."), *GetTransformFriendlyName(), StringCast<TCHAR>(exception.what()).Get());
+			UE_LOG(LogOpenColorIO, Error, TEXT("Failed to fetch shader info for color transform %s. Error message: %s"), *GetTransformFriendlyName(), StringCast<TCHAR>(exception.what()).Get());
 		}
 #endif
 	}
 	else
 	{
-		UE_LOG(LogOpenColorIO, Error, TEXT("Failed to fetch shader info for color transform %s. Configuration file was invalid."), *GetTransformFriendlyName());
+		UE_LOG(LogOpenColorIO, Error, TEXT("Failed to fetch shader info for color transform %s. Configuration file [%s] was invalid."), *GetTransformFriendlyName(), *ConfigurationOwner->ConfigurationFile.FilePath);
 	}
 
 	return false;
diff --git a/Engine/Plugins/Compositing/OpenColorIO/Source/OpenColorIO/Private/OpenColorIODisplayExtension.cpp b/Engine/Plugins/Compositing/OpenColorIO/Source/OpenColorIO/Private/OpenColorIODisplayExtension.cpp
--- a/Engine/Plugins/Compositing/OpenColorIO/Source/OpenColorIO/Private/OpenColorIODisplayExtension.cpp
+++ b/Engine/Plugins/Compositing/OpenColorIO/Source/OpenColorIO/Private/OpenColorIODisplayExtension.cpp
@@ -28,34 +28,8 @@
 #include "PostProcess/PostProcessMaterial.h"
 
 
-BEGIN_SHADER_PARAMETER_STRUCT(FOpenColorIOErrorShaderParameters, )
-	SHADER_PARAMETER_RDG_TEXTURE(Texture2D, InputTexture)
-	SHADER_PARAMETER_SAMPLER(SamplerState, InputTextureSampler)
-	SHADER_PARAMETER_TEXTURE(Texture2D, MiniFontTexture)
-	RENDER_TARGET_BINDING_SLOTS()
-END_SHADER_PARAMETER_STRUCT()
-
-class FOpenColorIOErrorPassPS : public FGlobalShader
-{
-public:
-	DECLARE_GLOBAL_SHADER(FOpenColorIOErrorPassPS);
-	SHADER_USE_PARAMETER_STRUCT(FOpenColorIOErrorPassPS, FGlobalShader);
-
-	using FParameters = FOpenColorIOErrorShaderParameters;
-};
-
-IMPLEMENT_GLOBAL_SHADER(FOpenColorIOErrorPassPS, "/Plugin/OpenColorIO/Private/OpenColorIOErrorShader.usf", "MainPS", SF_Pixel);
-
-
 float FOpenColorIODisplayExtension::DefaultDisplayGamma = 2.2f;
 
-namespace {
-	FRHITexture* GetSystemMiniFontTexture()
-	{
-		return GSystemTextures.AsciiTexture ? GSystemTextures.AsciiTexture->GetRHI() : GSystemTextures.WhiteDummy->GetRHI();
-	}
-};
-
 FOpenColorIODisplayExtension::FOpenColorIODisplayExtension(const FAutoRegister& AutoRegister, FViewportClient* AssociatedViewportClient)
 	: FSceneViewExtensionBase(AutoRegister)
 	, LinkedViewportClient(AssociatedViewportClient)
@@ -87,11 +61,7 @@ void FOpenColorIODisplayExtension::SetupView(FSceneViewFamily& InViewFamily, FSc
 	FOpenColorIOTransformResource* ShaderResource = nullptr;
 	TSortedMap<int32, FTextureResource*> TransformTextureResources;
 
-	if (DisplayConfiguration.ColorConfiguration.ConfigurationSource == nullptr)
-	{
-		UE_LOG(LogOpenColorIO, Warning, TEXT("Can't apply display look - Invalid config asset"));
-	}
-	else
+	if (DisplayConfiguration.ColorConfiguration.ConfigurationSource != nullptr)
 	{
 		const bool bFoundTransform = DisplayConfiguration.ColorConfiguration.ConfigurationSource->GetRenderResources(
 			InViewFamily.GetFeatureLevel()
@@ -99,17 +69,13 @@ void FOpenColorIODisplayExtension::SetupView(FSceneViewFamily& InViewFamily, FSc
 			, ShaderResource
 			, TransformTextureResources);
 
-		if (!bFoundTransform)
-		{
-			UE_LOG(LogOpenColorIO, Warning, TEXT("Can't apply display look - Couldn't find shader to transform: %s"), *DisplayConfiguration.ColorConfiguration.ToString());
-		}
-		else
+		if (bFoundTransform)
 		{
 			// Transform was found, so shader must be there but doesn't mean the actual shader is available
 			check(ShaderResource);
 			if (ShaderResource->GetShaderGameThread<FOpenColorIOPixelShader>().IsNull())
 			{
-				UE_LOG(LogOpenColorIO, Warning, TEXT("Can't apply display look - Shader was invalid for Resource %s"), *ShaderResource->GetFriendlyName());
+				ensureMsgf(false, TEXT("Can't apply display look - Shader was invalid for Resource %s"), *ShaderResource->GetFriendlyName());
 
 				//Invalidate shader resource
 				ShaderResource = nullptr;
@@ -187,7 +153,7 @@ FScreenPassTexture FOpenColorIODisplayExtension::PostProcessPassAfterTonemap_Ren
 		FOpenColorIOErrorShaderParameters* Parameters = GraphBuilder.AllocParameters<FOpenColorIOErrorShaderParameters>();
 		Parameters->InputTexture = SceneColor.Texture;
 		Parameters->InputTextureSampler = TStaticSamplerState<>::GetRHI();
-		Parameters->MiniFontTexture = GetSystemMiniFontTexture();
+		Parameters->MiniFontTexture = OpenColorIOGetMiniFontTexture();
 		Parameters->RenderTargets[0] = Output.GetRenderTargetBinding();
 
 		AddDrawScreenPass(GraphBuilder, RDG_EVENT_NAME("OCIODisplayLookError"), ViewInfo, OutputViewport, InputViewport, OCIOPixelShader, Parameters);
diff --git a/Engine/Plugins/Compositing/OpenColorIO/Source/OpenColorIO/Private/OpenColorIORendering.cpp b/Engine/Plugins/Compositing/OpenColorIO/Source/OpenColorIO/Private/OpenColorIORendering.cpp
--- a/Engine/Plugins/Compositing/OpenColorIO/Source/OpenColorIO/Private/OpenColorIORendering.cpp
+++ b/Engine/Plugins/Compositing/OpenColorIO/Source/OpenColorIO/Private/OpenColorIORendering.cpp
@@ -8,6 +8,7 @@
 #include "GlobalShader.h"
 #include "Logging/LogMacros.h"
 #include "Logging/MessageLog.h"
+#include "OpenColorIOColorSpace.h"
 #include "OpenColorIOConfiguration.h"
 #include "OpenColorIOModule.h"
 #include "OpenColorIOShader.h"
@@ -17,11 +18,27 @@
 #include "ScreenPass.h"
 #include "TextureResource.h"
 
+namespace {
+	FViewInfo CreateDummyViewInfo(const FIntRect& InViewRect)
+	{
+		FSceneViewFamily ViewFamily(FSceneViewFamily::ConstructionValues(nullptr, nullptr, FEngineShowFlags(ESFIM_Game))
+			.SetTime(FGameTime())
+			.SetGammaCorrection(1.0f));
+		FSceneViewInitOptions ViewInitOptions;
+		ViewInitOptions.ViewFamily = &ViewFamily;
+		ViewInitOptions.SetViewRectangle(InViewRect);
+		ViewInitOptions.ViewOrigin = FVector::ZeroVector;
+		ViewInitOptions.ViewRotationMatrix = FMatrix::Identity;
+		ViewInitOptions.ProjectionMatrix = FMatrix::Identity;
+
+		return FViewInfo(ViewInitOptions);
+	}
+}
 
-void ProcessOCIOColorSpaceTransform_RenderThread(
+void AddDrawOCIOScreenPass(
 	FRHICommandListImmediate& InRHICmdList
-	, ERHIFeatureLevel::Type InFeatureLevel
-	, FOpenColorIOTransformResource* InOCIOColorTransformResource
+	, const ERHIFeatureLevel::Type InFeatureLevel
+	, FOpenColorIOTransformResource* InShaderResource
 	, const TSortedMap<int32, FTextureResource*>& InTextureResources
 	, FTextureRHIRef InputSpaceColorTexture
 	, FTextureRHIRef OutputSpaceColorTexture
@@ -35,30 +52,34 @@ void ProcessOCIOColorSpaceTransform_RenderThread(
 	FRDGTextureRef OutputTexture = GraphBuilder.RegisterExternalTexture(CreateRenderTarget(OutputSpaceColorTexture->GetTexture2D(), TEXT("OCIORenderTargetTexture")));
 	FScreenPassTextureViewport Viewport = FScreenPassTextureViewport(OutputTexture);
 	FScreenPassRenderTarget ScreenPassRenderTarget = FScreenPassRenderTarget(OutputTexture, FIntRect(FIntPoint::ZeroValue, OutputResolution), ERenderTargetLoadAction::EClear);
+	FViewInfo DummyView = CreateDummyViewInfo(ScreenPassRenderTarget.ViewRect);
+
+	if (InShaderResource != nullptr)
+	{
+		TShaderRef<FOpenColorIOPixelShader> OCIOPixelShader = InShaderResource->GetShader<FOpenColorIOPixelShader>();
+
+		FOpenColorIOPixelShaderParameters* Parameters = GraphBuilder.AllocParameters<FOpenColorIOPixelShaderParameters>();
+		Parameters->InputTexture = InputTexture;
+		Parameters->InputTextureSampler = TStaticSamplerState<>::GetRHI();
+		OpenColorIOBindTextureResources(Parameters, InTextureResources);
+		// Set Gamma to 1., since we do not have any display parameters or requirement for Gamma.
+		Parameters->Gamma = 1.0;
+		Parameters->RenderTargets[0] = ScreenPassRenderTarget.GetRenderTargetBinding();
+
+		AddDrawScreenPass(GraphBuilder, RDG_EVENT_NAME("DrawOCIOScreenPass"), DummyView, Viewport, Viewport, OCIOPixelShader, Parameters);
+	}
+	else
+	{
+		TShaderMapRef<FOpenColorIOErrorPassPS> OCIOPixelShader(GetGlobalShaderMap(InFeatureLevel));
 
-	TShaderRef<FOpenColorIOPixelShader> OCIOPixelShader = InOCIOColorTransformResource->GetShader<FOpenColorIOPixelShader>();
-
-	FOpenColorIOPixelShaderParameters* Parameters = GraphBuilder.AllocParameters<FOpenColorIOPixelShaderParameters>();
-	Parameters->InputTexture = InputTexture;
-	Parameters->InputTextureSampler = TStaticSamplerState<>::GetRHI();
-	OpenColorIOBindTextureResources(Parameters, InTextureResources);
-	// Set Gamma to 1., since we do not have any display parameters or requirement for Gamma.
-	Parameters->Gamma = 1.0;
-	Parameters->RenderTargets[0] = ScreenPassRenderTarget.GetRenderTargetBinding();
-
-	//Dummy ViewFamily/ViewInfo created for AddDrawScreenPass.
-	FSceneViewFamily ViewFamily(FSceneViewFamily::ConstructionValues(nullptr, nullptr, FEngineShowFlags(ESFIM_Game))
-		.SetTime(FGameTime())
-		.SetGammaCorrection(1.0f));
-	FSceneViewInitOptions ViewInitOptions;
-	ViewInitOptions.ViewFamily = &ViewFamily;
-	ViewInitOptions.SetViewRectangle(ScreenPassRenderTarget.ViewRect);
-	ViewInitOptions.ViewOrigin = FVector::ZeroVector;
-	ViewInitOptions.ViewRotationMatrix = FMatrix::Identity;
-	ViewInitOptions.ProjectionMatrix = FMatrix::Identity;
-	FViewInfo DummyView(ViewInitOptions);
-
-	AddDrawScreenPass(GraphBuilder, RDG_EVENT_NAME("ProcessOCIOColorSpaceXfrm"), DummyView, Viewport, Viewport, OCIOPixelShader, Parameters);
+		FOpenColorIOErrorShaderParameters* Parameters = GraphBuilder.AllocParameters<FOpenColorIOErrorShaderParameters>();
+		Parameters->InputTexture = InputTexture;
+		Parameters->InputTextureSampler = TStaticSamplerState<>::GetRHI();
+		Parameters->MiniFontTexture = OpenColorIOGetMiniFontTexture();
+		Parameters->RenderTargets[0] = ScreenPassRenderTarget.GetRenderTargetBinding();
+
+		AddDrawScreenPass(GraphBuilder, RDG_EVENT_NAME("DrawOCIOErrorScreenPass"), DummyView, Viewport, Viewport, OCIOPixelShader, Parameters);
+	}
 
 	GraphBuilder.Execute();
 }
@@ -68,70 +89,62 @@ bool FOpenColorIORendering::ApplyColorTransform(UWorld* InWorld, const FOpenColo
 {
 	check(IsInGameThread());
 
-	if (InSettings.ConfigurationSource == nullptr)
+	if (!ensureMsgf(InTexture, TEXT("Can't apply color transform - Invalid Input Texture")))
 	{
-		UE_LOG(LogOpenColorIO, Warning, TEXT("Can't apply color transform - Invalid config asset"));
 		return false;
 	}
 
-	if (InTexture == nullptr)
+	if (!ensureMsgf(OutRenderTarget, TEXT("Can't apply color transform - Invalid Output Texture")))
 	{
-		UE_LOG(LogOpenColorIO, Warning, TEXT("Can't apply color transform - Invalid Input Texture"));
 		return false;
 	}
 
-	if (OutRenderTarget == nullptr)
-	{
-		UE_LOG(LogOpenColorIO, Warning, TEXT("Can't apply color transform - Invalid Output Texture"));
-		return false;
-	}
-
-
 	FTextureResource* InputResource = InTexture->GetResource();
 	FTextureResource* OutputResource = OutRenderTarget->GetResource();
-	if (InputResource == nullptr)
+	if (!ensureMsgf(InputResource, TEXT("Can't apply color transform - Invalid Input Texture resource")))
 	{
-		UE_LOG(LogOpenColorIO, Warning, TEXT("Can't apply color transform - Invalid Input Texture resource"));
 		return false;
 	}
 
-	if (OutputResource == nullptr)
+	if (!ensureMsgf(OutputResource, TEXT("Can't apply color transform - Invalid Output Texture resource")))
 	{
-		UE_LOG(LogOpenColorIO, Warning, TEXT("Can't apply color transform - Invalid Output Texture resource"));
 		return false;
 	}
 
 	const ERHIFeatureLevel::Type FeatureLevel = InWorld->Scene->GetFeatureLevel();
 	FOpenColorIOTransformResource* ShaderResource = nullptr;
-	TSortedMap<int32, FTextureResource*> TextureResources;
-	bool bFoundTransform = InSettings.ConfigurationSource->GetRenderResources(FeatureLevel, InSettings, ShaderResource, TextureResources);
-	if (!bFoundTransform)
-	{
-		UE_LOG(LogOpenColorIO, Warning, TEXT("Can't apply color transform - Couldn't find shader to transform: %s"), *InSettings.ToString());
-		return false;
-	}
-
-	check(ShaderResource);
+	TSortedMap<int32, FTextureResource*> TransformTextureResources;
 
-	if (ShaderResource->GetShaderGameThread<FOpenColorIOPixelShader>().IsNull())
+	if (InSettings.ConfigurationSource != nullptr)
 	{
-		UE_LOG(LogOpenColorIO, Warning, TEXT("OCIOPass - Shader was invalid for Resource %s"), *ShaderResource->GetFriendlyName());
-		return false;
-	}
-
+		const bool bFoundTransform = InSettings.ConfigurationSource->GetRenderResources(FeatureLevel, InSettings, ShaderResource, TransformTextureResources);
 
+		if (bFoundTransform)
+		{
+			check(ShaderResource);
+			if (ShaderResource->GetShaderGameThread<FOpenColorIOPixelShader>().IsNull())
+			{
+				ensureMsgf(false, TEXT("Can't apply display look - Shader was invalid for Resource %s"), *ShaderResource->GetFriendlyName());
+
+				//Invalidate shader resource
+				ShaderResource = nullptr;
+			}
+		}
+	}
+	
 	ENQUEUE_RENDER_COMMAND(ProcessColorSpaceTransform)(
-		[FeatureLevel, InputResource, OutputResource, ShaderResource, TextureResources](FRHICommandListImmediate& RHICmdList)
+		[FeatureLevel, InputResource, OutputResource, ShaderResource, TextureResources = MoveTemp(TransformTextureResources)](FRHICommandListImmediate& RHICmdList)
 		{
-			ProcessOCIOColorSpaceTransform_RenderThread(
-			RHICmdList,
-			FeatureLevel,
-			ShaderResource,
-			TextureResources,
-			InputResource->TextureRHI,
-			OutputResource->TextureRHI,
-			FIntPoint(OutputResource->GetSizeX(), OutputResource->GetSizeY()));
+			AddDrawOCIOScreenPass(
+				RHICmdList,
+				FeatureLevel,
+				ShaderResource,
+				TextureResources,
+				InputResource->TextureRHI,
+				OutputResource->TextureRHI,
+				FIntPoint(OutputResource->GetSizeX(), OutputResource->GetSizeY()));
 		}
 	);
-	return true;
+
+	return ShaderResource != nullptr;
 }
diff --git a/Engine/Plugins/Compositing/OpenColorIO/Source/OpenColorIO/Private/OpenColorIOShader.cpp b/Engine/Plugins/Compositing/OpenColorIO/Source/OpenColorIO/Private/OpenColorIOShader.cpp
--- a/Engine/Plugins/Compositing/OpenColorIO/Source/OpenColorIO/Private/OpenColorIOShader.cpp
+++ b/Engine/Plugins/Compositing/OpenColorIO/Source/OpenColorIO/Private/OpenColorIOShader.cpp
@@ -6,6 +6,7 @@
 #include "OpenColorIOModule.h"
 #include "OpenColorIOShared.h"
 #include "TextureResource.h"
+#include "SystemTextures.h"
 
 void OpenColorIOBindTextureResources(FOpenColorIOPixelShaderParameters* Parameters, const TSortedMap<int32, FTextureResource*>& InTextureResources)
 {
@@ -35,4 +36,10 @@ void OpenColorIOBindTextureResources(FOpenColorIOPixelShaderParameters* Paramete
 	}
 }
 
+FRHITexture* OpenColorIOGetMiniFontTexture()
+{
+	return GSystemTextures.AsciiTexture ? GSystemTextures.AsciiTexture->GetRHI() : GSystemTextures.WhiteDummy->GetRHI();
+}
+
 IMPLEMENT_SHADER_TYPE(, FOpenColorIOPixelShader, TEXT("/Plugin/OpenColorIO/Private/OpenColorIOShader.usf"), TEXT("MainPS"), SF_Pixel)
+IMPLEMENT_SHADER_TYPE(, FOpenColorIOErrorPassPS, TEXT("/Plugin/OpenColorIO/Private/OpenColorIOErrorShader.usf"), TEXT("MainPS"), SF_Pixel);
diff --git a/Engine/Plugins/Compositing/OpenColorIO/Source/OpenColorIO/Public/OpenColorIOShader.h b/Engine/Plugins/Compositing/OpenColorIO/Source/OpenColorIO/Public/OpenColorIOShader.h
--- a/Engine/Plugins/Compositing/OpenColorIO/Source/OpenColorIO/Public/OpenColorIOShader.h
+++ b/Engine/Plugins/Compositing/OpenColorIO/Source/OpenColorIO/Public/OpenColorIOShader.h
@@ -15,6 +15,7 @@
 
 class FTextureResource;
 class UClass;
+class FRHITexture;
 
 BEGIN_SHADER_PARAMETER_STRUCT(FOpenColorIOPixelShaderParameters, )
 	SHADER_PARAMETER_RDG_TEXTURE(Texture2D, InputTexture)
@@ -82,4 +83,22 @@ public:
 	}
 };
 
+BEGIN_SHADER_PARAMETER_STRUCT(FOpenColorIOErrorShaderParameters, )
+	SHADER_PARAMETER_RDG_TEXTURE(Texture2D, InputTexture)
+	SHADER_PARAMETER_SAMPLER(SamplerState, InputTextureSampler)
+	SHADER_PARAMETER_TEXTURE(Texture2D, MiniFontTexture)
+	RENDER_TARGET_BINDING_SLOTS()
+END_SHADER_PARAMETER_STRUCT()
+
+class OPENCOLORIO_API FOpenColorIOErrorPassPS : public FGlobalShader
+{
+public:
+	DECLARE_GLOBAL_SHADER(FOpenColorIOErrorPassPS);
+	SHADER_USE_PARAMETER_STRUCT(FOpenColorIOErrorPassPS, FGlobalShader);
+
+	using FParameters = FOpenColorIOErrorShaderParameters;
+};
+
 OPENCOLORIO_API void OpenColorIOBindTextureResources(FOpenColorIOPixelShaderParameters* Parameters, const TSortedMap<int32, FTextureResource*>& InTextureResources);
+
+FRHITexture* OpenColorIOGetMiniFontTexture();
