OCIO: Combine RDG passes into a single function.

diff --git a/Engine/Plugins/Compositing/OpenColorIO/Source/OpenColorIO/Private/OpenColorIODisplayExtension.cpp b/Engine/Plugins/Compositing/OpenColorIO/Source/OpenColorIO/Private/OpenColorIODisplayExtension.cpp
--- a/Engine/Plugins/Compositing/OpenColorIO/Source/OpenColorIO/Private/OpenColorIODisplayExtension.cpp
+++ b/Engine/Plugins/Compositing/OpenColorIO/Source/OpenColorIO/Private/OpenColorIODisplayExtension.cpp
@@ -8,15 +8,9 @@
 #include "OpenColorIOColorTransform.h"
 #include "OpenColorIODisplayManager.h"
 #include "OpenColorIOModule.h"
-#include "SceneView.h"
 
-#include "GlobalShader.h"
 #include "OpenColorIOShader.h"
-#include "PipelineStateCache.h"
-#include "RenderTargetPool.h"
 #include "Shader.h"
-#include "ShaderParameters.h"
-#include "ShaderParameterStruct.h"
 
 #include "RHI.h"
 #include "SceneView.h"
@@ -125,40 +119,19 @@ FScreenPassTexture FOpenColorIODisplayExtension::PostProcessPassAfterTonemap_Ren
 		Output = FScreenPassRenderTarget::CreateFromInput(GraphBuilder, SceneColor, ViewInfo.GetOverwriteLoadAction(), TEXT("OCIORenderTarget"));
 	}
 
-	const FScreenPassTextureViewport InputViewport(SceneColor);
-	const FScreenPassTextureViewport OutputViewport(Output);
-
-	if (CachedResourcesRenderThread.ShaderResource)
-	{
-		TShaderRef<FOpenColorIOPixelShader> OCIOPixelShader = CachedResourcesRenderThread.ShaderResource->GetShader<FOpenColorIOPixelShader>();
-
-		const float DisplayGamma = View.Family->RenderTarget->GetDisplayGamma();
-
-		FOpenColorIOPixelShaderParameters* Parameters = GraphBuilder.AllocParameters<FOpenColorIOPixelShaderParameters>();
-		Parameters->InputTexture = SceneColor.Texture;
-		Parameters->InputTextureSampler = TStaticSamplerState<>::GetRHI();
-		OpenColorIOBindTextureResources(Parameters, CachedResourcesRenderThread.TextureResources);
-
-		// There is a special case where post processing and tonemapper are disabled. In this case tonemapper applies a static display Inverse of Gamma which defaults to 2.2.
-		// In the case when Both PostProcessing and ToneMapper are disabled we apply gamma manually. In every other case we apply inverse gamma before applying OCIO.
-		Parameters->Gamma = (View.Family->EngineShowFlags.Tonemapper == 0) || (View.Family->EngineShowFlags.PostProcessing == 0) ? DefaultDisplayGamma : DefaultDisplayGamma / DisplayGamma;
-		Parameters->RenderTargets[0] = Output.GetRenderTargetBinding();
-
-		AddDrawScreenPass(GraphBuilder, RDG_EVENT_NAME("OCIODisplayLook"), ViewInfo, OutputViewport, InputViewport, OCIOPixelShader, Parameters);
-	}
-	else
-	{
-		// Fallback error pass, printing OCIO error message indicators across the viewport. (Helpful to quickly identify an OCIO config issue on nDisplay for example.)
-		TShaderMapRef<FOpenColorIOErrorPassPS> OCIOPixelShader(ViewInfo.ShaderMap);
-		FOpenColorIOErrorShaderParameters* Parameters = GraphBuilder.AllocParameters<FOpenColorIOErrorShaderParameters>();
-		Parameters->InputTexture = SceneColor.Texture;
-		Parameters->InputTextureSampler = TStaticSamplerState<>::GetRHI();
-		Parameters->MiniFontTexture = OpenColorIOGetMiniFontTexture();
-		Parameters->RenderTargets[0] = Output.GetRenderTargetBinding();
-
-		AddDrawScreenPass(GraphBuilder, RDG_EVENT_NAME("OCIODisplayLookError"), ViewInfo, OutputViewport, InputViewport, OCIOPixelShader, Parameters);
-	}
-
+	const float EngineDisplayGamma = View.Family->RenderTarget->GetDisplayGamma();
+	// There is a special case where post processing and tonemapper are disabled. In this case tonemapper applies a static display Inverse of Gamma which defaults to 2.2.
+	// In the case when Both PostProcessing and ToneMapper are disabled we apply gamma manually. In every other case we apply inverse gamma before applying OCIO.
+	float DisplayGamma = (View.Family->EngineShowFlags.Tonemapper == 0) || (View.Family->EngineShowFlags.PostProcessing == 0) ? DefaultDisplayGamma : DefaultDisplayGamma / EngineDisplayGamma;
+
+	FOpenColorIORendering::AddPass_RenderThread(
+		GraphBuilder,
+		ViewInfo,
+		SceneColor,
+		Output,
+		CachedResourcesRenderThread,
+		DisplayGamma
+	);
 
 	return MoveTemp(Output);
 }
\ No newline at end of file
diff --git a/Engine/Plugins/Compositing/OpenColorIO/Source/OpenColorIO/Private/OpenColorIORendering.cpp b/Engine/Plugins/Compositing/OpenColorIO/Source/OpenColorIO/Private/OpenColorIORendering.cpp
--- a/Engine/Plugins/Compositing/OpenColorIO/Source/OpenColorIO/Private/OpenColorIORendering.cpp
+++ b/Engine/Plugins/Compositing/OpenColorIO/Source/OpenColorIO/Private/OpenColorIORendering.cpp
@@ -35,53 +35,45 @@ namespace {
 	}
 }
 
-void AddDrawOCIOScreenPass(
-	FRHICommandListImmediate& InRHICmdList
-	, const ERHIFeatureLevel::Type InFeatureLevel
-	, FOpenColorIOTransformResource* InShaderResource
-	, const TSortedMap<int32, FTextureResource*>& InTextureResources
-	, FTextureRHIRef InputSpaceColorTexture
-	, FTextureRHIRef OutputSpaceColorTexture
-	, FIntPoint OutputResolution)
+// static
+void FOpenColorIORendering::AddPass_RenderThread(
+	FRDGBuilder& GraphBuilder,
+	const FViewInfo& View,
+	const FScreenPassTexture& Input,
+	const FScreenPassRenderTarget& Output,
+	const FOpenColorIORenderPassResources& InPassResource,
+	float InGamma)
 {
 	check(IsInRenderingThread());
 
-	FRDGBuilder GraphBuilder(InRHICmdList);
-
-	FRDGTextureRef InputTexture = GraphBuilder.RegisterExternalTexture(CreateRenderTarget(InputSpaceColorTexture->GetTexture2D(), TEXT("OCIOInputTexture")));
-	FRDGTextureRef OutputTexture = GraphBuilder.RegisterExternalTexture(CreateRenderTarget(OutputSpaceColorTexture->GetTexture2D(), TEXT("OCIORenderTargetTexture")));
-	FScreenPassTextureViewport Viewport = FScreenPassTextureViewport(OutputTexture);
-	FScreenPassRenderTarget ScreenPassRenderTarget = FScreenPassRenderTarget(OutputTexture, FIntRect(FIntPoint::ZeroValue, OutputResolution), ERenderTargetLoadAction::EClear);
-	FViewInfo DummyView = CreateDummyViewInfo(ScreenPassRenderTarget.ViewRect);
+	const FScreenPassTextureViewport InputViewport(Input);
+	const FScreenPassTextureViewport OutputViewport(Output);
 
-	if (InShaderResource != nullptr)
+	if (InPassResource.ShaderResource != nullptr)
 	{
-		TShaderRef<FOpenColorIOPixelShader> OCIOPixelShader = InShaderResource->GetShader<FOpenColorIOPixelShader>();
+		TShaderRef<FOpenColorIOPixelShader> OCIOPixelShader = InPassResource.ShaderResource->GetShader<FOpenColorIOPixelShader>();
 
 		FOpenColorIOPixelShaderParameters* Parameters = GraphBuilder.AllocParameters<FOpenColorIOPixelShaderParameters>();
-		Parameters->InputTexture = InputTexture;
+		Parameters->InputTexture = Input.Texture;
 		Parameters->InputTextureSampler = TStaticSamplerState<>::GetRHI();
-		OpenColorIOBindTextureResources(Parameters, InTextureResources);
-		// Set Gamma to 1., since we do not have any display parameters or requirement for Gamma.
-		Parameters->Gamma = 1.0;
-		Parameters->RenderTargets[0] = ScreenPassRenderTarget.GetRenderTargetBinding();
+		OpenColorIOBindTextureResources(Parameters, InPassResource.TextureResources);
+		Parameters->Gamma = InGamma;
+		Parameters->RenderTargets[0] = Output.GetRenderTargetBinding();
 
-		AddDrawScreenPass(GraphBuilder, RDG_EVENT_NAME("DrawOCIOScreenPass"), DummyView, Viewport, Viewport, OCIOPixelShader, Parameters);
+		AddDrawScreenPass(GraphBuilder, RDG_EVENT_NAME("OpenColorIOPass"), View, OutputViewport, InputViewport, OCIOPixelShader, Parameters);
 	}
 	else
 	{
-		TShaderMapRef<FOpenColorIOErrorPassPS> OCIOPixelShader(GetGlobalShaderMap(InFeatureLevel));
-
+		// Fallback error pass, printing OCIO error message indicators across the viewport. (Helpful to quickly identify an OCIO config issue on nDisplay for example.)
+		TShaderMapRef<FOpenColorIOErrorPassPS> OCIOPixelShader(View.ShaderMap);
 		FOpenColorIOErrorShaderParameters* Parameters = GraphBuilder.AllocParameters<FOpenColorIOErrorShaderParameters>();
-		Parameters->InputTexture = InputTexture;
+		Parameters->InputTexture = Input.Texture;
 		Parameters->InputTextureSampler = TStaticSamplerState<>::GetRHI();
 		Parameters->MiniFontTexture = OpenColorIOGetMiniFontTexture();
-		Parameters->RenderTargets[0] = ScreenPassRenderTarget.GetRenderTargetBinding();
+		Parameters->RenderTargets[0] = Output.GetRenderTargetBinding();
 
-		AddDrawScreenPass(GraphBuilder, RDG_EVENT_NAME("DrawOCIOErrorScreenPass"), DummyView, Viewport, Viewport, OCIOPixelShader, Parameters);
+		AddDrawScreenPass(GraphBuilder, RDG_EVENT_NAME("OpenColorIOErrorPass"), View, OutputViewport, InputViewport, OCIOPixelShader, Parameters);
 	}
-
-	GraphBuilder.Execute();
 }
 
 // static
@@ -133,16 +125,24 @@ bool FOpenColorIORendering::ApplyColorTransform(UWorld* InWorld, const FOpenColo
 	}
 	
 	ENQUEUE_RENDER_COMMAND(ProcessColorSpaceTransform)(
-		[FeatureLevel, InputResource, OutputResource, ShaderResource, TextureResources = MoveTemp(TransformTextureResources)](FRHICommandListImmediate& RHICmdList)
+		[InputResource, OutputResource, ShaderResource, TextureResources = MoveTemp(TransformTextureResources)](FRHICommandListImmediate& RHICmdList)
 		{
-			AddDrawOCIOScreenPass(
-				RHICmdList,
-				FeatureLevel,
-				ShaderResource,
-				TextureResources,
-				InputResource->TextureRHI,
-				OutputResource->TextureRHI,
-				FIntPoint(OutputResource->GetSizeX(), OutputResource->GetSizeY()));
+			FRDGBuilder GraphBuilder(RHICmdList);
+
+			FRDGTextureRef InputTexture = GraphBuilder.RegisterExternalTexture(CreateRenderTarget(InputResource->TextureRHI, TEXT("OCIOInputTexture")));
+			FRDGTextureRef OutputTexture = GraphBuilder.RegisterExternalTexture(CreateRenderTarget(OutputResource->TextureRHI, TEXT("OCIORenderTargetTexture")));
+			FIntPoint  OutputResolution = FIntPoint(OutputResource->GetSizeX(), OutputResource->GetSizeY());
+			FScreenPassRenderTarget Output = FScreenPassRenderTarget(OutputTexture, FIntRect(FIntPoint::ZeroValue, OutputResolution), ERenderTargetLoadAction::EClear);
+
+			AddPass_RenderThread(
+				GraphBuilder,
+				CreateDummyViewInfo(Output.ViewRect),
+				FScreenPassTexture(InputTexture),
+				Output,
+				FOpenColorIORenderPassResources{ ShaderResource, TextureResources},
+				1.0f); // Set Gamma to 1., since we do not have any display parameters or requirement for Gamma.
+
+			GraphBuilder.Execute();
 		}
 	);
 
diff --git a/Engine/Plugins/Compositing/OpenColorIO/Source/OpenColorIO/Public/OpenColorIODisplayExtension.h b/Engine/Plugins/Compositing/OpenColorIO/Source/OpenColorIO/Public/OpenColorIODisplayExtension.h
--- a/Engine/Plugins/Compositing/OpenColorIO/Source/OpenColorIO/Public/OpenColorIODisplayExtension.h
+++ b/Engine/Plugins/Compositing/OpenColorIO/Source/OpenColorIO/Public/OpenColorIODisplayExtension.h
@@ -6,13 +6,12 @@
 
 #include "CoreMinimal.h"
 #include "OpenColorIOColorSpace.h"
+#include "OpenColorIORendering.h"
 #include "UObject/GCObject.h"
 
 
 class FViewportClient;
-class FOpenColorIOTransformResource;
 class FSceneViewFamily;
-class FTextureResource;
 
 #define OPENCOLORIO_SCENE_VIEW_EXTENSION_PRIORITY 100
 
@@ -62,14 +61,8 @@ public:
 	static float DefaultDisplayGamma;
 private:
 
-	struct FCachedRenderResource
-	{
-		FOpenColorIOTransformResource* ShaderResource = nullptr;
-		TSortedMap<int32, FTextureResource*> TextureResources;
-	};
-
-	/** Cached resources required to apply conversion for render thread */
-	FCachedRenderResource CachedResourcesRenderThread;
+	/** Cached pass resources required to apply conversion for render thread */
+	FOpenColorIORenderPassResources CachedResourcesRenderThread;
 
 	/** Configuration to apply during post render callback */
 	FOpenColorIODisplayConfiguration DisplayConfiguration;
diff --git a/Engine/Plugins/Compositing/OpenColorIO/Source/OpenColorIO/Public/OpenColorIORendering.h b/Engine/Plugins/Compositing/OpenColorIO/Source/OpenColorIO/Public/OpenColorIORendering.h
--- a/Engine/Plugins/Compositing/OpenColorIO/Source/OpenColorIO/Public/OpenColorIORendering.h
+++ b/Engine/Plugins/Compositing/OpenColorIO/Source/OpenColorIO/Public/OpenColorIORendering.h
@@ -4,18 +4,28 @@
 
 #include "CoreMinimal.h"
 
-#include "OpenColorIOColorSpace.h"
-
-
+struct FOpenColorIOColorConversionSettings;
+class FOpenColorIOTransformResource;
+class FRDGBuilder;
+struct FScreenPassRenderTarget;
+struct FScreenPassTexture;
+class FTextureResource;
+class FViewInfo;
 class UTexture;
 class UTextureRenderTarget2D;
 class UWorld;
 
+struct OPENCOLORIO_API FOpenColorIORenderPassResources
+{
+	FOpenColorIOTransformResource* ShaderResource = nullptr;
+	TSortedMap<int32, FTextureResource*> TextureResources = {};
+};
 
 /** Entry point to trigger OpenColorIO conversion rendering */
 class OPENCOLORIO_API FOpenColorIORendering
 {
 public:
+	FOpenColorIORendering() = delete;
 
 	/**
 	 * Applies the color transform described in the settings
@@ -28,6 +38,21 @@ public:
 	 */
 	static bool ApplyColorTransform(UWorld* InWorld, const FOpenColorIOColorConversionSettings& InSettings, UTexture* InTexture, UTextureRenderTarget2D* OutRenderTarget);
 
-private:
-	FOpenColorIORendering() {}
+	/**
+	 * Applies the color transform RDG pass with the provided resources.
+	 *
+	 * @param GraphBuilder Render graph builder
+	 * @param View Scene view with additional information
+	 * @param Input Input color texture
+	 * @param Output Destination render target
+	 * @param InPassInfo OpenColorIO shader and texture resources
+	 * @param InGamma Display gamma
+	 */
+	static void AddPass_RenderThread(
+		FRDGBuilder& GraphBuilder,
+		const FViewInfo& View,
+		const FScreenPassTexture& Input,
+		const FScreenPassRenderTarget& Output,
+		const FOpenColorIORenderPassResources& InPassInfo,
+		float InGamma);
 };
\ No newline at end of file
