OCIO: Renaming "error" pass to "invalid" for consistency with menus.

#rnx

diff --git a/Engine/Plugins/Compositing/OpenColorIO/Shaders/Private/OpenColorIOErrorShader.usf b/Engine/Plugins/Compositing/OpenColorIO/Shaders/Private/OpenColorIOErrorShader.usf
deleted file mode 100644
diff --git a/Engine/Plugins/Compositing/OpenColorIO/Shaders/Private/OpenColorIOInvalidShader.usf b/Engine/Plugins/Compositing/OpenColorIO/Shaders/Private/OpenColorIOInvalidShader.usf
new file mode 100644
--- /dev/null
+++ b/Engine/Plugins/Compositing/OpenColorIO/Shaders/Private/OpenColorIOInvalidShader.usf
@@ -0,0 +1,28 @@
+// Copyright Epic Games, Inc. All Rights Reserved.
+
+/*=============================================================================
+ OpenColorIOErrorShader.usf: Fallback OCIO shader when there is an error with the transform.
+ =============================================================================*/
+
+#include "/Engine/Public/Platform.ush"
+#include "/Engine/Private/MiniFontCommon.ush"
+
+Texture2D		InputTexture;
+SamplerState	InputTextureSampler;
+
+void MainPS(in noperspective float4 InUVAndScreenPos : TEXCOORD0,
+	float4 SvPosition : SV_POSITION,
+	out float4 OutColor : SV_Target0)
+{
+	const int2 PixelPos = (int2)SvPosition.xy % int2(256, 256); // Tile message every 256px
+	const float3 TextColor = float3(1, 0.75, 0.25);
+	int2 Cursor = int2(128, 128);
+
+	OutColor = InputTexture.Sample(InputTextureSampler, InUVAndScreenPos.xy);
+
+#define P(x)  PrintCharacter(PixelPos, OutColor.xyz, TextColor, Cursor, x);
+
+	P(_O_)P(_C_)P(_I_)P(_O_)
+	P(_SPC_)
+	P(_I_)P(_N_)P(_V_)P(_A_)P(_L_)P(_I_)P(_D_)
+}
\ No newline at end of file
diff --git a/Engine/Plugins/Compositing/OpenColorIO/Source/OpenColorIO/Private/OpenColorIORendering.cpp b/Engine/Plugins/Compositing/OpenColorIO/Source/OpenColorIO/Private/OpenColorIORendering.cpp
--- a/Engine/Plugins/Compositing/OpenColorIO/Source/OpenColorIO/Private/OpenColorIORendering.cpp
+++ b/Engine/Plugins/Compositing/OpenColorIO/Source/OpenColorIO/Private/OpenColorIORendering.cpp
@@ -64,15 +64,15 @@ void FOpenColorIORendering::AddPass_RenderThread(
 	}
 	else
 	{
-		// Fallback error pass, printing OCIO error message indicators across the viewport. (Helpful to quickly identify an OCIO config issue on nDisplay for example.)
-		TShaderMapRef<FOpenColorIOErrorPassPS> OCIOPixelShader(View.ShaderMap);
-		FOpenColorIOErrorShaderParameters* Parameters = GraphBuilder.AllocParameters<FOpenColorIOErrorShaderParameters>();
+		// Fallback pass, printing invalid message across the viewport.
+		TShaderMapRef<FOpenColorIOInvalidPixelShader> OCIOInvalidPixelShader(View.ShaderMap);
+		FOpenColorIOInvalidShaderParameters* Parameters = GraphBuilder.AllocParameters<FOpenColorIOInvalidShaderParameters>();
 		Parameters->InputTexture = Input.Texture;
 		Parameters->InputTextureSampler = TStaticSamplerState<>::GetRHI();
 		Parameters->MiniFontTexture = OpenColorIOGetMiniFontTexture();
 		Parameters->RenderTargets[0] = Output.GetRenderTargetBinding();
 
-		AddDrawScreenPass(GraphBuilder, RDG_EVENT_NAME("OpenColorIOErrorPass"), View, OutputViewport, InputViewport, OCIOPixelShader, Parameters);
+		AddDrawScreenPass(GraphBuilder, RDG_EVENT_NAME("OpenColorIOInvalidPass"), View, OutputViewport, InputViewport, OCIOInvalidPixelShader, Parameters);
 	}
 }
 
diff --git a/Engine/Plugins/Compositing/OpenColorIO/Source/OpenColorIO/Private/OpenColorIOShader.cpp b/Engine/Plugins/Compositing/OpenColorIO/Source/OpenColorIO/Private/OpenColorIOShader.cpp
--- a/Engine/Plugins/Compositing/OpenColorIO/Source/OpenColorIO/Private/OpenColorIOShader.cpp
+++ b/Engine/Plugins/Compositing/OpenColorIO/Source/OpenColorIO/Private/OpenColorIOShader.cpp
@@ -42,4 +42,4 @@ FRHITexture* OpenColorIOGetMiniFontTexture()
 }
 
 IMPLEMENT_SHADER_TYPE(, FOpenColorIOPixelShader, TEXT("/Plugin/OpenColorIO/Private/OpenColorIOShader.usf"), TEXT("MainPS"), SF_Pixel)
-IMPLEMENT_SHADER_TYPE(, FOpenColorIOErrorPassPS, TEXT("/Plugin/OpenColorIO/Private/OpenColorIOErrorShader.usf"), TEXT("MainPS"), SF_Pixel);
+IMPLEMENT_SHADER_TYPE(, FOpenColorIOInvalidPixelShader, TEXT("/Plugin/OpenColorIO/Private/OpenColorIOInvalidShader.usf"), TEXT("MainPS"), SF_Pixel);
diff --git a/Engine/Plugins/Compositing/OpenColorIO/Source/OpenColorIO/Public/OpenColorIOShader.h b/Engine/Plugins/Compositing/OpenColorIO/Source/OpenColorIO/Public/OpenColorIOShader.h
--- a/Engine/Plugins/Compositing/OpenColorIO/Source/OpenColorIO/Public/OpenColorIOShader.h
+++ b/Engine/Plugins/Compositing/OpenColorIO/Source/OpenColorIO/Public/OpenColorIOShader.h
@@ -83,20 +83,20 @@ public:
 	}
 };
 
-BEGIN_SHADER_PARAMETER_STRUCT(FOpenColorIOErrorShaderParameters, )
+BEGIN_SHADER_PARAMETER_STRUCT(FOpenColorIOInvalidShaderParameters, )
 	SHADER_PARAMETER_RDG_TEXTURE(Texture2D, InputTexture)
 	SHADER_PARAMETER_SAMPLER(SamplerState, InputTextureSampler)
 	SHADER_PARAMETER_TEXTURE(Texture2D, MiniFontTexture)
 	RENDER_TARGET_BINDING_SLOTS()
 END_SHADER_PARAMETER_STRUCT()
 
-class OPENCOLORIO_API FOpenColorIOErrorPassPS : public FGlobalShader
+class OPENCOLORIO_API FOpenColorIOInvalidPixelShader : public FGlobalShader
 {
 public:
-	DECLARE_GLOBAL_SHADER(FOpenColorIOErrorPassPS);
-	SHADER_USE_PARAMETER_STRUCT(FOpenColorIOErrorPassPS, FGlobalShader);
+	DECLARE_GLOBAL_SHADER(FOpenColorIOInvalidPixelShader);
+	SHADER_USE_PARAMETER_STRUCT(FOpenColorIOInvalidPixelShader, FGlobalShader);
 
-	using FParameters = FOpenColorIOErrorShaderParameters;
+	using FParameters = FOpenColorIOInvalidShaderParameters;
 };
 
 OPENCOLORIO_API void OpenColorIOBindTextureResources(FOpenColorIOPixelShaderParameters* Parameters, const TSortedMap<int32, FTextureResource*>& InTextureResources);
