OCIO: Exposing support for context in the transform.

diff --git a/Engine/Plugins/Compositing/OpenColorIO/Source/OpenColorIO/Private/OpenColorIOColorTransform.cpp b/Engine/Plugins/Compositing/OpenColorIO/Source/OpenColorIO/Private/OpenColorIOColorTransform.cpp
--- a/Engine/Plugins/Compositing/OpenColorIO/Source/OpenColorIO/Private/OpenColorIOColorTransform.cpp
+++ b/Engine/Plugins/Compositing/OpenColorIO/Source/OpenColorIO/Private/OpenColorIOColorTransform.cpp
@@ -42,22 +42,34 @@ namespace {
 		return static_cast<OptimizationFlags>(OptimizationFlags::OPTIMIZATION_DEFAULT | OptimizationFlags::OPTIMIZATION_NO_DYNAMIC_PROPERTIES);
 	}
 
-	OCIO_NAMESPACE::ConstProcessorRcPtr GetTransformProcessor(UOpenColorIOColorTransform* InTransform, const OCIO_NAMESPACE::ConstConfigRcPtr& CurrentConfig)
+	OCIO_NAMESPACE::ConstProcessorRcPtr GetTransformProcessor(UOpenColorIOColorTransform* InTransform, const OCIO_NAMESPACE::ConstConfigRcPtr& InConfig, const TMap<FString, FString>& InContext)
 	{
 		EOpenColorIOViewTransformDirection DisplayViewDirection;
+		OCIO_NAMESPACE::ContextRcPtr Context = InConfig->getCurrentContext()->createEditableCopy();
+
+		for (const TPair<FString, FString>& KeyValue : InContext)
+		{
+			Context->setStringVar(TCHAR_TO_ANSI(*KeyValue.Key), TCHAR_TO_ANSI(*KeyValue.Value));
+		}
 
 		if (InTransform->GetDisplayViewDirection(DisplayViewDirection))
 		{
 			OCIO_NAMESPACE::TransformDirection OcioDirection = static_cast<OCIO_NAMESPACE::TransformDirection>(DisplayViewDirection);
 
-			return CurrentConfig->getProcessor(StringCast<ANSICHAR>(*InTransform->SourceColorSpace).Get(),
+			return InConfig->getProcessor(
+				Context,
+				StringCast<ANSICHAR>(*InTransform->SourceColorSpace).Get(),
 				StringCast<ANSICHAR>(*InTransform->Display).Get(),
 				StringCast<ANSICHAR>(*InTransform->View).Get(),
 				OcioDirection);
 		}
 		else
 		{
-			return CurrentConfig->getProcessor(StringCast<ANSICHAR>(*InTransform->SourceColorSpace).Get(), StringCast<ANSICHAR>(*InTransform->DestinationColorSpace).Get());
+			return InConfig->getProcessor(
+				Context,
+				StringCast<ANSICHAR>(*InTransform->SourceColorSpace).Get(),
+				StringCast<ANSICHAR>(*InTransform->DestinationColorSpace).Get()
+			);
 		}
 	}
 #endif
@@ -165,17 +177,21 @@ UOpenColorIOColorTransform::UOpenColorIOColorTransform(const FObjectInitializer&
 {
 }
 
-bool UOpenColorIOColorTransform::Initialize(UOpenColorIOConfiguration* InOwner, const FString& InSourceColorSpace, const FString& InDestinationColorSpace)
+bool UOpenColorIOColorTransform::Initialize(UOpenColorIOConfiguration* InOwner, const FString& InSourceColorSpace, const FString& InDestinationColorSpace, const TMap<FString, FString>& InContextKeyValues)
 {
 	check(InOwner);
 	ConfigurationOwner = InOwner;
+	ContextKeyValues = InContextKeyValues;
+	
 	return GenerateColorTransformData(InSourceColorSpace, InDestinationColorSpace);
 }
 
-bool UOpenColorIOColorTransform::Initialize(UOpenColorIOConfiguration* InOwner, const FString& InSourceColorSpace, const FString& InDisplay, const FString& InView, EOpenColorIOViewTransformDirection InDirection)
+bool UOpenColorIOColorTransform::Initialize(UOpenColorIOConfiguration* InOwner, const FString& InSourceColorSpace, const FString& InDisplay, const FString& InView, EOpenColorIOViewTransformDirection InDirection, const TMap<FString, FString>& InContextKeyValues)
 {
 	check(InOwner);
 	ConfigurationOwner = InOwner;
+	ContextKeyValues = InContextKeyValues;
+
 	return GenerateColorTransformData(InSourceColorSpace, InDisplay, InView, InDirection);
 }
 
@@ -255,7 +271,7 @@ void UOpenColorIOColorTransform::CacheResourceTextures()
 			try
 #endif
 			{
-				OCIO_NAMESPACE::ConstProcessorRcPtr TransformProcessor = GetTransformProcessor(this, CurrentConfig);				
+				OCIO_NAMESPACE::ConstProcessorRcPtr TransformProcessor = GetTransformProcessor(this, CurrentConfig, ContextKeyValues);
 				if (TransformProcessor)
 				{
 					OCIO_NAMESPACE::GpuShaderDescRcPtr ShaderDescription = OCIO_NAMESPACE::GpuShaderDesc::CreateShaderDesc();
@@ -574,7 +590,7 @@ bool UOpenColorIOColorTransform::UpdateShaderInfo(FString& OutShaderCodeHash, FS
 		try
 #endif
 		{
-			OCIO_NAMESPACE::ConstProcessorRcPtr TransformProcessor = GetTransformProcessor(this, CurrentConfig);
+			OCIO_NAMESPACE::ConstProcessorRcPtr TransformProcessor = GetTransformProcessor(this, CurrentConfig, ContextKeyValues);
 			if (TransformProcessor)
 			{
 				OCIO_NAMESPACE::GpuShaderDescRcPtr ShaderDescription = OCIO_NAMESPACE::GpuShaderDesc::CreateShaderDesc();
diff --git a/Engine/Plugins/Compositing/OpenColorIO/Source/OpenColorIO/Public/OpenColorIOColorTransform.h b/Engine/Plugins/Compositing/OpenColorIO/Source/OpenColorIO/Public/OpenColorIOColorTransform.h
--- a/Engine/Plugins/Compositing/OpenColorIO/Source/OpenColorIO/Public/OpenColorIOColorTransform.h
+++ b/Engine/Plugins/Compositing/OpenColorIO/Source/OpenColorIO/Public/OpenColorIOColorTransform.h
@@ -32,8 +32,8 @@ public:
 	virtual ~UOpenColorIOColorTransform() {};
 
 public:
-	bool Initialize(UOpenColorIOConfiguration* InOwner, const FString& InSourceColorSpace, const FString& InDestinationColorSpace);
-	bool Initialize(UOpenColorIOConfiguration* InOwner, const FString& InSourceColorSpace, const FString& InDisplay, const FString& InView, EOpenColorIOViewTransformDirection InDirection);
+	bool Initialize(UOpenColorIOConfiguration* InOwner, const FString& InSourceColorSpace, const FString& InDestinationColorSpace, const TMap<FString, FString>& InContextKeyValues = {});
+	bool Initialize(UOpenColorIOConfiguration* InOwner, const FString& InSourceColorSpace, const FString& InDisplay, const FString& InView, EOpenColorIOViewTransformDirection InDirection, const TMap<FString, FString>& InContextKeyValues = {});
 
 	/**
 	 * Cache resource shaders for cooking on the given shader platform.
@@ -206,6 +206,9 @@ private:
 	
 	FOpenColorIOTransformResource* ColorTransformResources[ERHIFeatureLevel::Num];
 
+	/** Key-value string pairs used to define the processor's context. */
+	TMap<FString, FString> ContextKeyValues;
+
 	FRenderCommandFence ReleaseFence;
 
 #if WITH_EDITORONLY_DATA
