OCIO: Adding conversion setting change delegate, and equality operator support.

Preliminary changes required for sequencer track/shot integration.

diff --git a/Engine/Plugins/Compositing/OpenColorIO/Source/OpenColorIO/Public/OpenColorIOColorSpace.h b/Engine/Plugins/Compositing/OpenColorIO/Source/OpenColorIO/Public/OpenColorIOColorSpace.h
--- a/Engine/Plugins/Compositing/OpenColorIO/Source/OpenColorIO/Public/OpenColorIOColorSpace.h
+++ b/Engine/Plugins/Compositing/OpenColorIO/Source/OpenColorIO/Public/OpenColorIOColorSpace.h
@@ -123,6 +123,7 @@ struct OPENCOLORIO_API FOpenColorIOColorConversionSettings
 	GENERATED_BODY()
 
 public:
+	DECLARE_MULTICAST_DELEGATE(FOnConversionSettingsChange);
 
 	/** Default constructor. */
 	FOpenColorIOColorConversionSettings();
@@ -147,6 +148,8 @@ public:
 	UPROPERTY(EditAnywhere, BlueprintReadWrite, Category = ColorSpace)
 	EOpenColorIOViewTransformDirection DisplayViewDirection = EOpenColorIOViewTransformDirection::Forward;
 
+	/** Delegate triggered upon changes to the settings. */
+	FOnConversionSettingsChange& OnConversionSettingsChanged() { return ConversionSettingsChanged; }
 public:
 
 	/**
@@ -165,9 +168,24 @@ public:
 	*/
 	void ValidateColorSpaces();
 
+	bool operator==(const FOpenColorIOColorConversionSettings& Other) const { return Equals(Other); }
+	bool operator!=(const FOpenColorIOColorConversionSettings& Other) const { return !Equals(Other); }
+
+	/** Determines if this ColorConversionSettings is the same as another.*/
+	bool Equals(const FOpenColorIOColorConversionSettings& Other) const
+	{
+		return ConfigurationSource    == Other.ConfigurationSource
+			&& SourceColorSpace       == Other.SourceColorSpace
+			&& DestinationColorSpace  == Other.DestinationColorSpace
+			&& DestinationDisplayView == Other.DestinationDisplayView
+			&& DisplayViewDirection   == Other.DisplayViewDirection;
+	}
+
 private:
 	/** Whether or not these settings are of the display-view type. */
 	bool IsDisplayView() const;
+
+	FOnConversionSettingsChange ConversionSettingsChanged;
 };
 
 /**
diff --git a/Engine/Plugins/Compositing/OpenColorIO/Source/OpenColorIOEditor/Private/OpenColorIOColorSpaceConversionCustomization.cpp b/Engine/Plugins/Compositing/OpenColorIO/Source/OpenColorIOEditor/Private/OpenColorIOColorSpaceConversionCustomization.cpp
--- a/Engine/Plugins/Compositing/OpenColorIO/Source/OpenColorIOEditor/Private/OpenColorIOColorSpaceConversionCustomization.cpp
+++ b/Engine/Plugins/Compositing/OpenColorIO/Source/OpenColorIOEditor/Private/OpenColorIOColorSpaceConversionCustomization.cpp
@@ -83,6 +83,8 @@ void FOpenColorIOColorConversionSettingsCustomization::CustomizeChildren(TShared
 					{
 						TransformPicker[OCIO_Src]->SetConfiguration(ColorSpaceConversion->ConfigurationSource);
 						TransformPicker[OCIO_Dst]->SetConfiguration(ColorSpaceConversion->ConfigurationSource);
+
+						ColorSpaceConversion->OnConversionSettingsChanged().Broadcast();
 					}));
 			}
 			else if (ChildHandle->GetProperty()->GetFName() == GET_MEMBER_NAME_CHECKED(FOpenColorIOColorConversionSettings, SourceColorSpace))
@@ -175,6 +177,8 @@ void FOpenColorIOColorConversionSettingsCustomization::OnSourceColorSpaceChanged
 	TransformPicker[OCIO_Dst]->SetRestrictions(NewColorSpace, NewDisplayView);
 
 	ApplySelectionToConfiguration();
+
+	ColorSpaceConversion->OnConversionSettingsChanged().Broadcast();
 }
 
 void FOpenColorIOColorConversionSettingsCustomization::OnDestinationColorSpaceChanged(const FOpenColorIOColorSpace& NewColorSpace, const FOpenColorIODisplayView& NewDisplayView)
@@ -185,6 +189,8 @@ void FOpenColorIOColorConversionSettingsCustomization::OnDestinationColorSpaceCh
 	TransformPicker[OCIO_Src]->SetRestrictions(NewColorSpace, NewDisplayView);
 
 	ApplySelectionToConfiguration();
+
+	ColorSpaceConversion->OnConversionSettingsChanged().Broadcast();
 }
 
 void FOpenColorIOColorConversionSettingsCustomization::ApplyConfigurationToSelection()
@@ -254,6 +260,8 @@ void FOpenColorIOColorConversionSettingsCustomization::OnConfigurationReset()
 {
 	TransformPicker[OCIO_Src]->SetConfiguration(nullptr);
 	TransformPicker[OCIO_Dst]->SetConfiguration(nullptr);
+
+	ColorSpaceConversion->OnConversionSettingsChanged().Broadcast();
 }
 
 #undef LOCTEXT_NAMESPACE
diff --git a/Engine/Plugins/Compositing/OpenColorIO/Source/OpenColorIOEditor/Private/Widgets/SOpenColorIOColorSpacePicker.cpp b/Engine/Plugins/Compositing/OpenColorIO/Source/OpenColorIOEditor/Private/Widgets/SOpenColorIOColorSpacePicker.cpp
--- a/Engine/Plugins/Compositing/OpenColorIO/Source/OpenColorIOEditor/Private/Widgets/SOpenColorIOColorSpacePicker.cpp
+++ b/Engine/Plugins/Compositing/OpenColorIO/Source/OpenColorIOEditor/Private/Widgets/SOpenColorIOColorSpacePicker.cpp
@@ -84,7 +84,8 @@ void SOpenColorIOColorSpacePicker::SetConfiguration(TWeakObjectPtr<UOpenColorIOC
 	Configuration = NewConfiguration;
 
 	// Invalidate current color space selection
-	SetCurrentColorSpace(FOpenColorIOColorSpace());
+	ColorSpaceSelection.Reset();
+	DisplayViewSelection.Reset();
 }
 
 void SOpenColorIOColorSpacePicker::SetRestrictions(const FOpenColorIOColorSpace& InRestrictedColorSpace, const FOpenColorIODisplayView& InRestricedDisplayView)
