OpenColorIO: Alpha related fix in OCIO shader

diff --git a/Engine/Plugins/Compositing/OpenColorIO/Shaders/Private/OpenColorIOShader.usf b/Engine/Plugins/Compositing/OpenColorIO/Shaders/Private/OpenColorIOShader.usf
--- a/Engine/Plugins/Compositing/OpenColorIO/Shaders/Private/OpenColorIOShader.usf
+++ b/Engine/Plugins/Compositing/OpenColorIO/Shaders/Private/OpenColorIOShader.usf
@@ -36,7 +36,7 @@ void MainPS(in noperspective float4 InUVAndScreenPos : TEXCOORD0,
 	out float4 OutColor : SV_Target0)
 {
 	float4 InputColor = InputTexture.Sample(InputTextureSampler, InUVAndScreenPos.xy);
-	InputColor = max(0.0, InputColor);
-	InputColor = pow(InputColor, Gamma);
+	InputColor.xyz = max(0.0, InputColor.xyz);
+	InputColor.xyz = pow(InputColor.xyz, Gamma);
 	OutColor = float4(OCIOConvert(InputColor).xyz, InputColor.w);
 }
\ No newline at end of file
