OCIO: File path filter now also allows selection of *.ocioz files.

#rnx

diff --git a/Engine/Plugins/Compositing/OpenColorIO/Source/OpenColorIO/Public/OpenColorIOConfiguration.h b/Engine/Plugins/Compositing/OpenColorIO/Source/OpenColorIO/Public/OpenColorIOConfiguration.h
--- a/Engine/Plugins/Compositing/OpenColorIO/Source/OpenColorIO/Public/OpenColorIOConfiguration.h
+++ b/Engine/Plugins/Compositing/OpenColorIO/Source/OpenColorIO/Public/OpenColorIOConfiguration.h
@@ -101,7 +101,7 @@ private:
 
 public:
 
-	UPROPERTY(BlueprintReadWrite, EditAnywhere, Category = "Config", meta = (FilePathFilter = "ocio", RelativeToGameDir))
+	UPROPERTY(BlueprintReadWrite, EditAnywhere, Category = "Config", meta = (FilePathFilter = "Config Files (*.ocio, *.ocioz)|*.ocio;*.ocioz", RelativeToGameDir))
 	FFilePath ConfigurationFile;
 
 	UPROPERTY(BlueprintReadWrite, EditAnywhere, Category = "ColorSpace", meta=(OCIOConfigFile="ConfigurationFile"))
