OCIO: Converted asset action to new asset definition.

#rnx

diff --git a/Engine/Plugins/Compositing/OpenColorIO/Source/OpenColorIOEditor/OpenColorIOEditor.Build.cs b/Engine/Plugins/Compositing/OpenColorIO/Source/OpenColorIOEditor/OpenColorIOEditor.Build.cs
--- a/Engine/Plugins/Compositing/OpenColorIO/Source/OpenColorIOEditor/OpenColorIOEditor.Build.cs
+++ b/Engine/Plugins/Compositing/OpenColorIO/Source/OpenColorIOEditor/OpenColorIOEditor.Build.cs
@@ -8,6 +8,7 @@ namespace UnrealBuildTool.Rules
 		{
 			PrivateDependencyModuleNames.AddRange(
 				new string[] {
+					"AssetDefinition",
 					"Core",
 					"CoreUObject",
 					"DesktopWidgets",
diff --git a/Engine/Plugins/Compositing/OpenColorIO/Source/OpenColorIOEditor/Private/AssetDefinition_OpenColorIOConfiguration.cpp b/Engine/Plugins/Compositing/OpenColorIO/Source/OpenColorIOEditor/Private/AssetDefinition_OpenColorIOConfiguration.cpp
new file mode 100644
--- /dev/null
+++ b/Engine/Plugins/Compositing/OpenColorIO/Source/OpenColorIOEditor/Private/AssetDefinition_OpenColorIOConfiguration.cpp
@@ -0,0 +1,36 @@
+// Copyright Epic Games, Inc. All Rights Reserved.
+
+#include "AssetDefinition_OpenColorIOConfiguration.h"
+
+#include "OpenColorIOConfiguration.h"
+
+#define LOCTEXT_NAMESPACE "AssetDefinition_OpenColorIOConfiguration"
+
+
+/* FAssetTypeActions_Base interface
+*****************************************************************************/
+
+FText UAssetDefinition_OpenColorIOConfiguration::GetAssetDescription(const struct FAssetData& AssetData) const
+{
+	return FText::FromString(AssetData.GetTagValueRef<FString>(GET_MEMBER_NAME_CHECKED(UOpenColorIOConfiguration, ConfigurationFile)));
+}
+
+
+TSoftClassPtr<UObject> UAssetDefinition_OpenColorIOConfiguration::GetAssetClass() const
+{
+	return UOpenColorIOConfiguration::StaticClass();
+}
+
+
+FText UAssetDefinition_OpenColorIOConfiguration::GetAssetDisplayName() const
+{
+	return LOCTEXT("AssetTypeActions_OpenColorIOConfiguration", "OpenColorIO Configuration");
+}
+
+FLinearColor UAssetDefinition_OpenColorIOConfiguration::GetAssetColor() const
+{
+	return FLinearColor::White;
+}
+
+
+#undef LOCTEXT_NAMESPACE
diff --git a/Engine/Plugins/Compositing/OpenColorIO/Source/OpenColorIOEditor/Private/AssetDefinition_OpenColorIOConfiguration.h b/Engine/Plugins/Compositing/OpenColorIO/Source/OpenColorIOEditor/Private/AssetDefinition_OpenColorIOConfiguration.h
new file mode 100644
--- /dev/null
+++ b/Engine/Plugins/Compositing/OpenColorIO/Source/OpenColorIOEditor/Private/AssetDefinition_OpenColorIOConfiguration.h
@@ -0,0 +1,28 @@
+// Copyright Epic Games, Inc. All Rights Reserved.
+
+#pragma once
+
+#include "AssetDefinitionDefault.h"
+
+#include "AssetDefinition_OpenColorIOConfiguration.generated.h"
+
+UCLASS()
+class UAssetDefinition_OpenColorIOConfiguration : public UAssetDefinitionDefault
+{
+	GENERATED_BODY()
+
+public:
+	// Begin UAssetDefinition interface
+	virtual TConstArrayView<FAssetCategoryPath> GetAssetCategories() const override
+	{
+		static const auto Categories = { EAssetCategoryPaths::Misc };
+		return Categories;
+	}
+	virtual bool CanImport() const override { return false; }
+	
+	virtual FText GetAssetDisplayName() const override;
+	virtual FText GetAssetDescription(const FAssetData& AssetData) const override;
+	virtual TSoftClassPtr<UObject> GetAssetClass() const override;
+	virtual FLinearColor GetAssetColor() const override;
+	// End UAssetDefinition interface
+};
diff --git a/Engine/Plugins/Compositing/OpenColorIO/Source/OpenColorIOEditor/Private/AssetTypeActions_OpenColorIOConfiguration.cpp b/Engine/Plugins/Compositing/OpenColorIO/Source/OpenColorIOEditor/Private/AssetTypeActions_OpenColorIOConfiguration.cpp
deleted file mode 100644
diff --git a/Engine/Plugins/Compositing/OpenColorIO/Source/OpenColorIOEditor/Private/AssetTypeActions_OpenColorIOConfiguration.h b/Engine/Plugins/Compositing/OpenColorIO/Source/OpenColorIOEditor/Private/AssetTypeActions_OpenColorIOConfiguration.h
deleted file mode 100644
diff --git a/Engine/Plugins/Compositing/OpenColorIO/Source/OpenColorIOEditor/Private/OpenColorIOEditorModule.cpp b/Engine/Plugins/Compositing/OpenColorIO/Source/OpenColorIOEditor/Private/OpenColorIOEditorModule.cpp
--- a/Engine/Plugins/Compositing/OpenColorIO/Source/OpenColorIOEditor/Private/OpenColorIOEditorModule.cpp
+++ b/Engine/Plugins/Compositing/OpenColorIO/Source/OpenColorIOEditor/Private/OpenColorIOEditorModule.cpp
@@ -2,7 +2,6 @@
 
 #include "OpenColorIOEditorModule.h"
 
-#include "AssetTypeActions_OpenColorIOConfiguration.h"
 #include "Engine/World.h"
 #include "Interfaces/IPluginManager.h"
 #include "IOpenColorIOModule.h"
@@ -33,13 +32,6 @@ void FOpenColorIOEditorModule::StartupModule()
 	LLM_SCOPE_BYTAG(OpenColorIO_OpenColorIOEditor);
 
 	FWorldDelegates::OnPreWorldInitialization.AddRaw(this, &FOpenColorIOEditorModule::OnWorldInit);
-
-	// Register asset type actions for OpenColorIOConfiguration class
-	IAssetTools& AssetTools = FModuleManager::LoadModuleChecked<FAssetToolsModule>("AssetTools").Get();
-	TSharedRef<IAssetTypeActions> OpenColorIOConfigurationAssetTypeAction = MakeShared<FAssetTypeActions_OpenColorIOConfiguration>();
-	AssetTools.RegisterAssetTypeActions(OpenColorIOConfigurationAssetTypeAction);
-	RegisteredAssetTypeActions.Add(OpenColorIOConfigurationAssetTypeAction);
-
 	FCoreDelegates::OnFEngineLoopInitComplete.AddRaw(this, &FOpenColorIOEditorModule::OnEngineLoopInitComplete);
 	
 	RegisterCustomizations();
@@ -55,19 +47,6 @@ void FOpenColorIOEditorModule::ShutdownModule()
 	UnregisterStyle();
 	UnregisterCustomizations();
 
-	// Unregister AssetTypeActions
-	FAssetToolsModule* AssetToolsModule = FModuleManager::GetModulePtr<FAssetToolsModule>("AssetTools");
-
-	if (AssetToolsModule != nullptr)
-	{
-		IAssetTools& AssetTools = AssetToolsModule->Get();
-
-		for (const TSharedRef<IAssetTypeActions>& Action : RegisteredAssetTypeActions)
-		{
-			AssetTools.UnregisterAssetTypeActions(Action);
-		}
-	}
-
 	CleanFeatureLevelDelegate();
 	FWorldDelegates::OnPreWorldInitialization.RemoveAll(this);
 	FCoreDelegates::OnFEngineLoopInitComplete.RemoveAll(this);
