OCIO: Avoid SVE ensure trigger while resources aren't ready.

#rnx

diff --git a/Engine/Plugins/Compositing/OpenColorIO/Source/OpenColorIO/Private/OpenColorIOColorTransform.cpp b/Engine/Plugins/Compositing/OpenColorIO/Source/OpenColorIO/Private/OpenColorIOColorTransform.cpp
--- a/Engine/Plugins/Compositing/OpenColorIO/Source/OpenColorIO/Private/OpenColorIOColorTransform.cpp
+++ b/Engine/Plugins/Compositing/OpenColorIO/Source/OpenColorIO/Private/OpenColorIOColorTransform.cpp
@@ -463,6 +463,53 @@ bool UOpenColorIOColorTransform::GetRenderResources(ERHIFeatureLevel::Type InFea
 	return OutShaderResource != nullptr;
 }
 
+bool UOpenColorIOColorTransform::AreRenderResourcesReady() const
+{
+	check(IsInGameThread());
+
+	int32 NumShadersRequired = 0;
+	int32 NumShadersReady = 0;
+	for (const FOpenColorIOTransformResource* Resource : ColorTransformResources)
+	{
+		if (Resource)
+		{
+			NumShadersRequired++;
+
+			if (Resource->IsCompilationFinished())
+			{
+				NumShadersReady++;
+			}
+		}
+	}
+
+	// Transform shaders are always required, and the all of the required ones should have finished compiling.
+	if (NumShadersRequired == 0 || NumShadersRequired != NumShadersReady)
+	{
+		return false;
+	}
+
+	// Textures are optional, depending on the transform.
+	for (const TPair<int32, TObjectPtr<UTexture>>& TexturePair : Textures)
+	{
+		const TObjectPtr<UTexture>& Texture = TexturePair.Value;
+
+		if (Texture->GetResource() == nullptr)
+		{
+			return false;
+		}
+
+#if WITH_EDITOR
+		// Note: this check is valid the first time a texture is created, but wouldn't be relevant if we updated existing textures.
+		if (Texture->IsCompiling())
+		{
+			return false;
+		}
+#endif
+	}
+
+	return true;
+}
+
 bool UOpenColorIOColorTransform::GetShaderAndLUTResouces(ERHIFeatureLevel::Type InFeatureLevel, FOpenColorIOTransformResource*& OutShaderResource, FTextureResource*& OutLUT3dResource)
 {
 	TSortedMap<int32, FTextureResource*> OutTextureResources;
diff --git a/Engine/Plugins/Compositing/OpenColorIO/Source/OpenColorIO/Private/OpenColorIOConfiguration.cpp b/Engine/Plugins/Compositing/OpenColorIO/Source/OpenColorIO/Private/OpenColorIOConfiguration.cpp
--- a/Engine/Plugins/Compositing/OpenColorIO/Source/OpenColorIO/Private/OpenColorIOConfiguration.cpp
+++ b/Engine/Plugins/Compositing/OpenColorIO/Source/OpenColorIO/Private/OpenColorIOConfiguration.cpp
@@ -59,29 +59,24 @@ void UOpenColorIOConfiguration::BeginDestroy()
 	Super::BeginDestroy();
 }
 
-bool UOpenColorIOConfiguration::GetRenderResources(ERHIFeatureLevel::Type InFeatureLevel, const FOpenColorIOColorConversionSettings& InSettings, FOpenColorIOTransformResource*& OutShaderResource, TSortedMap<int32, FTextureResource*>& OutTextureResources)
+bool UOpenColorIOConfiguration::IsTransformReady(const FOpenColorIOColorConversionSettings& InSettings)
 {
-	TObjectPtr<UOpenColorIOColorTransform>* TransformPtr = ColorTransforms.FindByPredicate([&](const UOpenColorIOColorTransform* InTransform)
+	if (const TObjectPtr<UOpenColorIOColorTransform>* TransformPtr = FindTransform(InSettings))
 	{
-		EOpenColorIOViewTransformDirection DisplayViewDirection;
+		return (*TransformPtr)->AreRenderResourcesReady();
+	}
 
-		if (InTransform->GetDisplayViewDirection(DisplayViewDirection))
-		{
-			return InTransform->SourceColorSpace == InSettings.SourceColorSpace.ColorSpaceName && InTransform->Display == InSettings.DestinationDisplayView.Display && InTransform->View == InSettings.DestinationDisplayView.View && InSettings.DisplayViewDirection == DisplayViewDirection;
-		}
-		else
-		{
-			return InTransform->SourceColorSpace == InSettings.SourceColorSpace.ColorSpaceName && InTransform->DestinationColorSpace == InSettings.DestinationColorSpace.ColorSpaceName;
-		}
-	});
+	return false;
+}
 
-	if (TransformPtr == nullptr)
+bool UOpenColorIOConfiguration::GetRenderResources(ERHIFeatureLevel::Type InFeatureLevel, const FOpenColorIOColorConversionSettings& InSettings, FOpenColorIOTransformResource*& OutShaderResource, TSortedMap<int32, FTextureResource*>& OutTextureResources)
+{
+	if (const TObjectPtr<UOpenColorIOColorTransform>* TransformPtr = FindTransform(InSettings))
 	{
-		return false;
+		return (*TransformPtr)->GetRenderResources(InFeatureLevel, OutShaderResource, OutTextureResources);
 	}
 
-	UOpenColorIOColorTransform* Transform = *TransformPtr;
-	return Transform->GetRenderResources(InFeatureLevel, OutShaderResource, OutTextureResources);
+	return false;
 }
 
 bool UOpenColorIOConfiguration::GetShaderAndLUTResources(ERHIFeatureLevel::Type InFeatureLevel, const FString& InSourceColorSpace, const FString& InDestinationColorSpace, FOpenColorIOTransformResource*& OutShaderResource, FTextureResource*& OutLUT3dResource)
@@ -325,6 +320,23 @@ void UOpenColorIOConfiguration::StopDirectoryWatch()
 #endif
 }
 
+const TObjectPtr<UOpenColorIOColorTransform>* UOpenColorIOConfiguration::FindTransform(const FOpenColorIOColorConversionSettings& InSettings) const
+{
+	return ColorTransforms.FindByPredicate([&](const UOpenColorIOColorTransform* InTransform)
+		{
+			EOpenColorIOViewTransformDirection DisplayViewDirection;
+			
+			if (InTransform->GetDisplayViewDirection(DisplayViewDirection))
+			{
+				return InTransform->SourceColorSpace == InSettings.SourceColorSpace.ColorSpaceName && InTransform->Display == InSettings.DestinationDisplayView.Display && InTransform->View == InSettings.DestinationDisplayView.View && InSettings.DisplayViewDirection == DisplayViewDirection;
+			}
+			else
+			{
+				return InTransform->SourceColorSpace == InSettings.SourceColorSpace.ColorSpaceName && InTransform->DestinationColorSpace == InSettings.DestinationColorSpace.ColorSpaceName;
+			}
+		});
+}
+
 void UOpenColorIOConfiguration::CreateColorTransform(const FString& InSourceColorSpace, const FString& InDestinationColorSpace)
 {
 	if (InSourceColorSpace.IsEmpty() || InDestinationColorSpace.IsEmpty())
diff --git a/Engine/Plugins/Compositing/OpenColorIO/Source/OpenColorIO/Private/OpenColorIODisplayExtension.cpp b/Engine/Plugins/Compositing/OpenColorIO/Source/OpenColorIO/Private/OpenColorIODisplayExtension.cpp
--- a/Engine/Plugins/Compositing/OpenColorIO/Source/OpenColorIO/Private/OpenColorIODisplayExtension.cpp
+++ b/Engine/Plugins/Compositing/OpenColorIO/Source/OpenColorIO/Private/OpenColorIODisplayExtension.cpp
@@ -34,7 +34,10 @@ bool FOpenColorIODisplayExtension::IsActiveThisFrame_Internal(const FSceneViewEx
 {
 	if (Context.Viewport && LinkedViewportClient == Context.Viewport->GetClient() && DisplayConfiguration.bIsEnabled)
 	{
-		return DisplayConfiguration.ColorConfiguration.IsValid();
+		if (DisplayConfiguration.ColorConfiguration.IsValid())
+		{
+			return DisplayConfiguration.ColorConfiguration.ConfigurationSource->IsTransformReady(DisplayConfiguration.ColorConfiguration);
+		}
 	}
 
 	return false;
diff --git a/Engine/Plugins/Compositing/OpenColorIO/Source/OpenColorIO/Public/OpenColorIOColorTransform.h b/Engine/Plugins/Compositing/OpenColorIO/Source/OpenColorIO/Public/OpenColorIOColorTransform.h
--- a/Engine/Plugins/Compositing/OpenColorIO/Source/OpenColorIO/Public/OpenColorIOColorTransform.h
+++ b/Engine/Plugins/Compositing/OpenColorIO/Source/OpenColorIO/Public/OpenColorIOColorTransform.h
@@ -68,6 +68,11 @@ public:
 	 */
 	bool GetRenderResources(ERHIFeatureLevel::Type InFeatureLevel, FOpenColorIOTransformResource*& OutShaderResource, TSortedMap<int32, FTextureResource*>& OutTextureResources);
 
+	/**
+	 * Returns true if shader/texture resources have finished compiling and are ready for use (to be called on the game thread).
+	 */
+	bool AreRenderResourcesReady() const;
+
 	UE_DEPRECATED(5.1, "GetShaderAndLUTResouces is deprecated, please use GetRenderResources instead.")
 	bool GetShaderAndLUTResouces(ERHIFeatureLevel::Type InFeatureLevel, FOpenColorIOTransformResource*& OutShaderResource, FTextureResource*& OutLUT3dResource);
 
diff --git a/Engine/Plugins/Compositing/OpenColorIO/Source/OpenColorIO/Public/OpenColorIOConfiguration.h b/Engine/Plugins/Compositing/OpenColorIO/Source/OpenColorIO/Public/OpenColorIOConfiguration.h
--- a/Engine/Plugins/Compositing/OpenColorIO/Source/OpenColorIO/Public/OpenColorIOConfiguration.h
+++ b/Engine/Plugins/Compositing/OpenColorIO/Source/OpenColorIO/Public/OpenColorIOConfiguration.h
@@ -38,6 +38,7 @@ public:
 	UOpenColorIOConfiguration(const FObjectInitializer& ObjectInitializer);
 
 public:
+	bool IsTransformReady(const FOpenColorIOColorConversionSettings& InSettings);
 	bool GetRenderResources(ERHIFeatureLevel::Type InFeatureLevel, const FOpenColorIOColorConversionSettings& InSettings, FOpenColorIOTransformResource*& OutShaderResource, TSortedMap<int32, FTextureResource*>& OutTextureResources);
 	UE_DEPRECATED(5.1, "GetShaderAndLUTResources is deprecated, please use GetRenderResources instead.")
 	bool GetShaderAndLUTResources(ERHIFeatureLevel::Type InFeatureLevel, const FString& InSourceColorSpace, const FString& InDestinationColorSpace, FOpenColorIOTransformResource*& OutShaderResource, FTextureResource*& OutLUT3dResource);
@@ -65,6 +66,8 @@ public:
 #endif
 
 protected:
+
+	const TObjectPtr<UOpenColorIOColorTransform>* FindTransform(const FOpenColorIOColorConversionSettings& InSettings) const;
 	void CreateColorTransform(const FString& InSourceColorSpace, const FString& InDestinationColorSpace);
 	void CreateColorTransform(const FString& InSourceColorSpace, const FString& InDisplay, const FString& InView, EOpenColorIOViewTransformDirection InDirection);
 	void CleanupTransforms();
