OCIO: Fix inconsistent preprocessor state by avoiding WITH_OCIO in the configuration public header, and not leaking any raw OCIO library objects.

#rnx

diff --git a/Engine/Plugins/Compositing/OpenColorIO/Source/OpenColorIO/Private/OpenColorIOColorTransform.cpp b/Engine/Plugins/Compositing/OpenColorIO/Source/OpenColorIO/Private/OpenColorIOColorTransform.cpp
--- a/Engine/Plugins/Compositing/OpenColorIO/Source/OpenColorIO/Private/OpenColorIOColorTransform.cpp
+++ b/Engine/Plugins/Compositing/OpenColorIO/Source/OpenColorIO/Private/OpenColorIOColorTransform.cpp
@@ -8,6 +8,7 @@
 #include "Modules/ModuleManager.h"
 #include "OpenColorIOConfiguration.h"
 #include "OpenColorIOModule.h"
+#include "OpenColorIONativeConfiguration.h"
 #include "OpenColorIOSettings.h"
 #include "UObject/UObjectIterator.h"
 #include "DataDrivenShaderPlatformInfo.h"
@@ -264,7 +265,7 @@ void UOpenColorIOColorTransform::CacheResourceTextures()
 	if (Textures.IsEmpty())
 	{
 #if WITH_EDITOR && WITH_OCIO
-		OCIO_NAMESPACE::ConstConfigRcPtr CurrentConfig = ConfigurationOwner->GetLoadedConfiguration();
+		OCIO_NAMESPACE::ConstConfigRcPtr CurrentConfig = ConfigurationOwner->GetNativeConfig_Internal()->Get();
 		if (CurrentConfig)
 		{
 #if !PLATFORM_EXCEPTIONS_DISABLED
@@ -630,7 +631,7 @@ bool UOpenColorIOColorTransform::UpdateShaderInfo(FString& OutShaderCodeHash, FS
 {
 #if WITH_EDITOR
 #if WITH_OCIO
-	OCIO_NAMESPACE::ConstConfigRcPtr CurrentConfig = ConfigurationOwner->GetLoadedConfiguration();
+	OCIO_NAMESPACE::ConstConfigRcPtr CurrentConfig = ConfigurationOwner->GetNativeConfig_Internal()->Get();
 	if (CurrentConfig)
 	{
 #if !PLATFORM_EXCEPTIONS_DISABLED
diff --git a/Engine/Plugins/Compositing/OpenColorIO/Source/OpenColorIO/Private/OpenColorIOConfiguration.cpp b/Engine/Plugins/Compositing/OpenColorIO/Source/OpenColorIO/Private/OpenColorIOConfiguration.cpp
--- a/Engine/Plugins/Compositing/OpenColorIO/Source/OpenColorIO/Private/OpenColorIOConfiguration.cpp
+++ b/Engine/Plugins/Compositing/OpenColorIO/Source/OpenColorIO/Private/OpenColorIOConfiguration.cpp
@@ -11,6 +11,7 @@
 #include "Modules/ModuleManager.h"
 #include "OpenColorIOColorTransform.h"
 #include "OpenColorIOModule.h"
+#include "OpenColorIONativeConfiguration.h"
 #include "OpenColorIOSettings.h"
 #include "TextureResource.h"
 #include "Widgets/Notifications/SNotificationList.h"
@@ -49,6 +50,7 @@ namespace OCIODirectoryWatcher
 
 UOpenColorIOConfiguration::UOpenColorIOConfiguration(const FObjectInitializer& ObjectInitializer)
 	: Super(ObjectInitializer)
+	, NativeConfig(MakePimpl<FOpenColorIONativeConfiguration>())
 {
 	
 }
@@ -145,7 +147,7 @@ bool UOpenColorIOConfiguration::Validate() const
 	if (!ConfigurationFile.FilePath.IsEmpty())
 	{
 		//When loading the configuration file, if any errors are detected, it will throw an exception. Thus, our pointer won't be valid.
-		return LoadedConfig != nullptr;
+		return NativeConfig->Get() != nullptr;
 	}
 
 	return false;
@@ -168,6 +170,7 @@ void UOpenColorIOConfiguration::ReloadExistingColorspaces()
 	CleanupTransforms();
 	LoadConfiguration();
 
+	OCIO_NAMESPACE::ConstConfigRcPtr LoadedConfig = NativeConfig->Get();
 	if (!LoadedConfig)
 	{
 		return;
@@ -282,6 +285,11 @@ void UOpenColorIOConfiguration::ConfigPathChangedEvent(const TArray<FFileChangeD
 #endif
 }
 
+FOpenColorIONativeConfiguration* UOpenColorIOConfiguration::GetNativeConfig_Internal() const
+{
+	return NativeConfig.Get();
+}
+
 void UOpenColorIOConfiguration::StartDirectoryWatch(const FString& FilePath)
 {
 #if WITH_EDITOR && WITH_OCIO
@@ -593,13 +601,11 @@ void UOpenColorIOConfiguration::LoadConfiguration()
 			FString ConfigurationFilePath = ConfigurationFile.FilePath;
 			if (ConfigurationFilePath.StartsWith(TEXT("ocio://")))
 			{
-				LoadedConfig = OCIO_NAMESPACE::Config::CreateFromFile(StringCast<ANSICHAR>(*ConfigurationFilePath).Get());
+				NativeConfig->Set(OCIO_NAMESPACE::Config::CreateFromFile(StringCast<ANSICHAR>(*ConfigurationFilePath).Get()));
 				UE_LOG(LogOpenColorIO, Verbose, TEXT("Loaded built-in OCIO configuration file %s"), *ConfigurationFilePath);
 			}
 			else
 			{
-				LoadedConfig.reset();
-
 				FString FullPath;
 				if (ConfigurationFilePath.Contains(TEXT("{Engine}")))
 				{
@@ -620,11 +626,14 @@ void UOpenColorIOConfiguration::LoadConfiguration()
 				if (NewConfig)
 				{
 					UE_LOG(LogOpenColorIO, Verbose, TEXT("Loaded OCIO configuration file %s"), *FullPath);
-					LoadedConfig = NewConfig;
+					NativeConfig->Set(NewConfig);
+					
 					StartDirectoryWatch(FullPath);
 				}
 				else
 				{
+					NativeConfig->Set(nullptr);
+
 					UE_LOG(LogOpenColorIO, Error, TEXT("Could not load OCIO configuration file %s. Verify that the path is good or that the file is valid."), *ConfigurationFile.FilePath);
 				}
 			}
diff --git a/Engine/Plugins/Compositing/OpenColorIO/Source/OpenColorIO/Private/OpenColorIONativeConfiguration.h b/Engine/Plugins/Compositing/OpenColorIO/Source/OpenColorIO/Private/OpenColorIONativeConfiguration.h
new file mode 100644
--- /dev/null
+++ b/Engine/Plugins/Compositing/OpenColorIO/Source/OpenColorIO/Private/OpenColorIONativeConfiguration.h
@@ -0,0 +1,28 @@
+// Copyright Epic Games, Inc. All Rights Reserved.
+
+#pragma once
+
+#if WITH_EDITOR && WITH_OCIO
+#include "OpenColorIO/OpenColorIO.h"
+#endif
+
+class FOpenColorIONativeConfiguration {
+public:
+
+#if WITH_EDITOR && WITH_OCIO
+	OCIO_NAMESPACE::ConstConfigRcPtr Get() const
+	{
+		return Config;
+	}
+
+	void Set(OCIO_NAMESPACE::ConstConfigRcPtr InConfig)
+	{
+		Config = InConfig;
+	}
+#endif
+
+private:
+#if WITH_EDITORONLY_DATA && WITH_OCIO
+	OCIO_NAMESPACE::ConstConfigRcPtr Config;
+#endif //WITH_EDITORONLY_DATA
+};
diff --git a/Engine/Plugins/Compositing/OpenColorIO/Source/OpenColorIO/Public/OpenColorIOConfiguration.h b/Engine/Plugins/Compositing/OpenColorIO/Source/OpenColorIO/Public/OpenColorIOConfiguration.h
--- a/Engine/Plugins/Compositing/OpenColorIO/Source/OpenColorIO/Public/OpenColorIOConfiguration.h
+++ b/Engine/Plugins/Compositing/OpenColorIO/Source/OpenColorIO/Public/OpenColorIOConfiguration.h
@@ -11,13 +11,10 @@
 #include "RHIDefinitions.h"
 #endif
 
-#if WITH_EDITOR && WITH_OCIO
-#include "OpenColorIO/OpenColorIO.h"
-#endif
-
 #include "OpenColorIOConfiguration.generated.h"
 
 
+class FOpenColorIONativeConfiguration;
 class FOpenColorIOTransformResource;
 class FTextureResource;
 class UOpenColorIOColorTransform;
@@ -58,13 +55,8 @@ public:
 	*/
 	void ConfigPathChangedEvent(const TArray<FFileChangeData>& InFileChanges, const FString InFileMountPath);
 
-#if WITH_EDITORONLY_DATA && WITH_OCIO
-	OCIO_NAMESPACE::ConstConfigRcPtr GetLoadedConfiguration() const { return LoadedConfig; }
-
-	UE_DEPRECATED(5.2, "GetLoadedConfigurationFile is deprecated, please use GetLoadedConfiguration instead.")
-	OCIO_NAMESPACE::ConstConfigRcPtr GetLoadedConfigurationFile() const { return GetLoadedConfiguration(); }
-#endif
-
+	/** Internal only: Replacement for previous `GetLoadedConfiguration()` and `GetLoadedConfigurationFile()` functions, returning the private implementation of the native OCIO config. */
+	FOpenColorIONativeConfiguration* GetNativeConfig_Internal() const;
 protected:
 
 	const TObjectPtr<UOpenColorIOColorTransform>* FindTransform(const FOpenColorIOColorConversionSettings& InSettings) const;
@@ -118,10 +110,6 @@ private:
 	UPROPERTY()
 	TArray<TObjectPtr<UOpenColorIOColorTransform>> ColorTransforms;
 
-#if WITH_EDITORONLY_DATA && WITH_OCIO
-	OCIO_NAMESPACE::ConstConfigRcPtr LoadedConfig;
-#endif //WITH_EDITORONLY_DATA
-
 private:
 	struct FOCIOConfigWatchedDirInfo
 	{
@@ -137,4 +125,7 @@ private:
 
 	/** Information about the currently watched directory. Helps us manage the directory change events. */
 	FOCIOConfigWatchedDirInfo WatchedDirectoryInfo;
+
+	/** Private implementation of the native OpenColorIO config object. */
+	TPimplPtr<FOpenColorIONativeConfiguration> NativeConfig;
 };
