OCIO: Sharing the same native config object between the configuration object and its slate customization.

#rnx

diff --git a/Engine/Plugins/Compositing/OpenColorIO/Source/OpenColorIO/Public/OpenColorIOConfiguration.h b/Engine/Plugins/Compositing/OpenColorIO/Source/OpenColorIO/Public/OpenColorIOConfiguration.h
--- a/Engine/Plugins/Compositing/OpenColorIO/Source/OpenColorIO/Public/OpenColorIOConfiguration.h
+++ b/Engine/Plugins/Compositing/OpenColorIO/Source/OpenColorIO/Public/OpenColorIOConfiguration.h
@@ -99,10 +99,10 @@ public:
 	UPROPERTY(BlueprintReadWrite, EditAnywhere, Category = "Config", meta = (FilePathFilter = "Config Files (*.ocio, *.ocioz)|*.ocio;*.ocioz", RelativeToGameDir))
 	FFilePath ConfigurationFile;
 
-	UPROPERTY(BlueprintReadWrite, EditAnywhere, Category = "ColorSpace", meta=(OCIOConfigFile="ConfigurationFile"))
+	UPROPERTY(BlueprintReadWrite, EditAnywhere, Category = "ColorSpace")
 	TArray<FOpenColorIOColorSpace> DesiredColorSpaces;
 
-	UPROPERTY(BlueprintReadWrite, EditAnywhere, Category = "ColorSpace", DisplayName="Desired Display-Views", meta = (OCIOConfigFile = "ConfigurationFile"))
+	UPROPERTY(BlueprintReadWrite, EditAnywhere, Category = "ColorSpace", DisplayName="Desired Display-Views")
 	TArray<FOpenColorIODisplayView> DesiredDisplayViews;
 
 private:
diff --git a/Engine/Plugins/Compositing/OpenColorIO/Source/OpenColorIOEditor/Private/OpenColorIOColorSpaceCustomization.cpp b/Engine/Plugins/Compositing/OpenColorIO/Source/OpenColorIOEditor/Private/OpenColorIOColorSpaceCustomization.cpp
--- a/Engine/Plugins/Compositing/OpenColorIO/Source/OpenColorIOEditor/Private/OpenColorIOColorSpaceCustomization.cpp
+++ b/Engine/Plugins/Compositing/OpenColorIO/Source/OpenColorIOEditor/Private/OpenColorIOColorSpaceCustomization.cpp
@@ -6,6 +6,9 @@
 #include "DetailWidgetRow.h"
 #include "IPropertyUtilities.h"
 #include "PropertyHandle.h"
+#include "OpenColorIOConfiguration.h"
+#include "OpenColorIO/Private/OpenColorIONativeConfiguration.h"
+#include "OpenColorIO/OpenColorIO.h"
 #include "Widgets/Input/SComboButton.h"
 
 #define LOCTEXT_NAMESPACE "OpenColorIOColorSpaceCustomization"
@@ -13,10 +16,8 @@
 void FOpenColorIOColorSpaceCustomization::CustomizeHeader(TSharedRef<IPropertyHandle> InPropertyHandle, FDetailWidgetRow& HeaderRow, IPropertyTypeCustomizationUtils& CustomizationUtils)
 {
 	//Reset internals
-	ConfigurationFileProperty.Reset();
-	CachedConfigFile.reset();
-	LoadedFilePath = FFilePath();
-	
+	CachedNativeConfig = nullptr;
+
 	CachedProperty = InPropertyHandle;
 
 	if (CachedProperty->GetNumPerObjectValues() == 1 && CachedProperty->IsValidHandle())
@@ -33,9 +34,6 @@ void FOpenColorIOColorSpaceCustomization::CustomizeHeader(TSharedRef<IPropertyHa
 		check(ColorSpaceValue);
 		TSharedPtr<IPropertyUtilities> PropertyUtils = CustomizationUtils.GetPropertyUtilities();
 
-		// Get the ConfigurationFile to read color space from
-		ConfigurationFileProperty = GetConfigurationFileProperty();
-
 		HeaderRow
 			.NameContent()
 			[
@@ -68,9 +66,7 @@ void FOpenColorIOColorSpaceCustomization::CustomizeHeader(TSharedRef<IPropertyHa
 void FOpenColorIODisplayViewCustomization::CustomizeHeader(TSharedRef<IPropertyHandle> InPropertyHandle, FDetailWidgetRow& HeaderRow, IPropertyTypeCustomizationUtils& CustomizationUtils)
 {
 	//Reset internals
-	ConfigurationFileProperty.Reset();
-	CachedConfigFile.reset();
-	LoadedFilePath = FFilePath();
+	CachedNativeConfig = nullptr;
 
 	CachedProperty = InPropertyHandle;
 
@@ -88,9 +84,6 @@ void FOpenColorIODisplayViewCustomization::CustomizeHeader(TSharedRef<IPropertyH
 		check(DisplayViewValue);
 		TSharedPtr<IPropertyUtilities> PropertyUtils = CustomizationUtils.GetPropertyUtilities();
 
-		// Get the ConfigurationFile to read color space from
-		ConfigurationFileProperty = GetConfigurationFileProperty();
-
 		HeaderRow
 			.NameContent()
 			[
@@ -120,106 +113,26 @@ void FOpenColorIODisplayViewCustomization::CustomizeHeader(TSharedRef<IPropertyH
 	}
 }
 
-TSharedPtr<IPropertyHandle> IPropertyTypeCustomizationOpenColorIO::GetConfigurationFileProperty() const
-{
-	if (CachedProperty.IsValid())
-	{
-		static FName NAME_ConfigFile = TEXT("OCIOConfigFile");
-
-		//Verify if we're in an array before accessing the property directly. Need to go deeper for array properties.
-		TSharedPtr<IPropertyHandle> ParentArrayHandle = CachedProperty->GetParentHandle();
-		TSharedPtr<IPropertyHandleArray> ParentHandleAsArray = ParentArrayHandle->AsArray();
-		if (ParentHandleAsArray.IsValid())
-		{
-			//Struct is in an array, parent class is access through more layers
-			if (ParentArrayHandle->GetProperty()->HasMetaData(NAME_ConfigFile))
-			{
-				const FString& ConfigFileVariableName = ParentArrayHandle->GetProperty()->GetMetaData(NAME_ConfigFile);
-				TSharedPtr<IPropertyHandle> ClassHandle = ParentArrayHandle->GetParentHandle()->GetParentHandle();
-				if (ClassHandle.IsValid() && !ConfigFileVariableName.IsEmpty())
-				{
-					return ClassHandle->GetChildHandle(*ConfigFileVariableName);
-				}
-			}
-		}
-		else
-		{
-			if (CachedProperty->HasMetaData(NAME_ConfigFile))
-			{
-				const FString& ConfigFileVariableName = CachedProperty->GetMetaData(NAME_ConfigFile);
-				TSharedPtr<IPropertyHandle> ClassHandle = CachedProperty->GetParentHandle();
-				if (ClassHandle.IsValid() && !ConfigFileVariableName.IsEmpty())
-				{
-					return ClassHandle->GetChildHandle(*ConfigFileVariableName);
-				}
-			}
-		}
-	}
-
-	return nullptr;
-}
-
-bool IPropertyTypeCustomizationOpenColorIO::LoadConfigurationFile(const FFilePath& InFilePath)
+IPropertyTypeCustomizationOpenColorIO::IPropertyTypeCustomizationOpenColorIO(TSharedPtr<IPropertyHandle> InConfigurationObjectProperty)
+	: ConfigurationObjectProperty(MoveTemp(InConfigurationObjectProperty))
 {
-#if WITH_OCIO
-#if !PLATFORM_EXCEPTIONS_DISABLED
-	try
-#endif
-	{
-		FString FullPath;
-		FString ConfigurationFilePath = InFilePath.FilePath;
-		if (ConfigurationFilePath.Contains(TEXT("{Engine}")))
-		{
-			ConfigurationFilePath = FPaths::ConvertRelativePathToFull(ConfigurationFilePath.Replace(TEXT("{Engine}"), *FPaths::EngineDir()));
-		}    
-
-		if (!FPaths::IsRelative(ConfigurationFilePath))
-		{
-			FullPath = ConfigurationFilePath;
-		}
-		else
-		{
-			const FString AbsoluteGameDir = FPaths::ConvertRelativePathToFull(FPaths::ProjectDir());
-			FullPath = FPaths::ConvertRelativePathToFull(FPaths::Combine(AbsoluteGameDir, ConfigurationFilePath));
-		}
-
-		CachedConfigFile = OCIO_NAMESPACE::Config::CreateFromFile(StringCast<ANSICHAR>(*FullPath).Get());
-		if (!CachedConfigFile)
-		{
-			return false;
-		}
-	}
-#if !PLATFORM_EXCEPTIONS_DISABLED
-	catch (OCIO_NAMESPACE::Exception&)
-	{
-		return false;
-	}
-#endif
-
-	return true;
-#else //WITH_OCIO
-	return false;
-#endif //WITH_OCIO
 }
 
-bool IPropertyTypeCustomizationOpenColorIO::CheckValidConfiguration()
+FOpenColorIONativeConfiguration* IPropertyTypeCustomizationOpenColorIO::GetNativeConfig(const TSharedPtr<IPropertyHandle>& InConfigurationAssetProperty)
 {
-	if (ConfigurationFileProperty.IsValid())
+	if (InConfigurationAssetProperty.IsValid())
 	{
 		TArray<void*> RawData;
-		ConfigurationFileProperty->AccessRawData(RawData);
-
+		InConfigurationAssetProperty->AccessRawData(RawData);
 		check(RawData.Num() == 1);
-		FFilePath* ConfigFilePath = reinterpret_cast<FFilePath*>(RawData[0]);
-		check(ConfigFilePath);
-
-		if (!ConfigFilePath->FilePath.IsEmpty() && ConfigFilePath->FilePath != LoadedFilePath.FilePath)
-		{
-			return LoadConfigurationFile(*ConfigFilePath);
-		}
+		
+		UOpenColorIOConfiguration* Configuration = reinterpret_cast<UOpenColorIOConfiguration*>(RawData[0]);
+		check(Configuration);
+		
+		return Configuration->GetNativeConfig_Internal();		
 	}
 
-	return false;
+	return nullptr;
 }
 
 void FOpenColorIOColorSpaceCustomization::ProcessColorSpaceForMenuGeneration(FMenuBuilder& InMenuBuilder, const int32 InMenuDepth, const FString& InPreviousFamilyHierarchy, const FOpenColorIOColorSpace& InColorSpace, TArray<FString>& InOutExistingMenuFilter)
@@ -257,11 +170,11 @@ void FOpenColorIOColorSpaceCustomization::PopulateSubMenu(FMenuBuilder& InMenuBu
 	// To keep track of submenus that were already added
 	TArray<FString> ExistingSubMenus;
 
-	const int32 ColorSpaceCount = CachedConfigFile->getNumColorSpaces(OCIO_NAMESPACE::SEARCH_REFERENCE_SPACE_ALL, OCIO_NAMESPACE::COLORSPACE_ACTIVE);
+	const int32 ColorSpaceCount = CachedNativeConfig->Get()->getNumColorSpaces(OCIO_NAMESPACE::SEARCH_REFERENCE_SPACE_ALL, OCIO_NAMESPACE::COLORSPACE_ACTIVE);
 	for (int32 i = 0; i < ColorSpaceCount; ++i)
 	{
-		const char* ColorSpaceName = CachedConfigFile->getColorSpaceNameByIndex(OCIO_NAMESPACE::SEARCH_REFERENCE_SPACE_ALL, OCIO_NAMESPACE::COLORSPACE_ACTIVE, i);
-		OCIO_NAMESPACE::ConstColorSpaceRcPtr LibColorSpace = CachedConfigFile->getColorSpace(ColorSpaceName);
+		const char* ColorSpaceName = CachedNativeConfig->Get()->getColorSpaceNameByIndex(OCIO_NAMESPACE::SEARCH_REFERENCE_SPACE_ALL, OCIO_NAMESPACE::COLORSPACE_ACTIVE, i);
+		OCIO_NAMESPACE::ConstColorSpaceRcPtr LibColorSpace = CachedNativeConfig->Get()->getColorSpace(ColorSpaceName);
 		if (!LibColorSpace)
 		{
 			continue;
@@ -317,7 +230,7 @@ void FOpenColorIOColorSpaceCustomization::AddMenuEntry(FMenuBuilder& InMenuBuild
 
 TSharedRef<SWidget> FOpenColorIOColorSpaceCustomization::HandleSourceComboButtonMenuContent()
 {
-	bool bValidConfiguration = CheckValidConfiguration();
+	CachedNativeConfig = GetNativeConfig(ConfigurationObjectProperty);
 
 	// Generate menu
 	FMenuBuilder MenuBuilder(true, nullptr);
@@ -325,13 +238,13 @@ TSharedRef<SWidget> FOpenColorIOColorSpaceCustomization::HandleSourceComboButton
 
 	MenuBuilder.BeginSection("AllColorSpaces", LOCTEXT("AllColorSpacesSection", "ColorSpaces"));
 	{
-		if (bValidConfiguration)
+		if (CachedNativeConfig)
 		{
-			const int32 ColorSpaceCount = CachedConfigFile->getNumColorSpaces(OCIO_NAMESPACE::SEARCH_REFERENCE_SPACE_ALL, OCIO_NAMESPACE::COLORSPACE_ACTIVE);
+			const int32 ColorSpaceCount = CachedNativeConfig->Get()->getNumColorSpaces(OCIO_NAMESPACE::SEARCH_REFERENCE_SPACE_ALL, OCIO_NAMESPACE::COLORSPACE_ACTIVE);
 			for (int32 i = 0; i < ColorSpaceCount; ++i)
 			{
-				const char* ColorSpaceName = CachedConfigFile->getColorSpaceNameByIndex(OCIO_NAMESPACE::SEARCH_REFERENCE_SPACE_ALL, OCIO_NAMESPACE::COLORSPACE_ACTIVE, i);
-				OCIO_NAMESPACE::ConstColorSpaceRcPtr LibColorSpace = CachedConfigFile->getColorSpace(ColorSpaceName);
+				const char* ColorSpaceName = CachedNativeConfig->Get()->getColorSpaceNameByIndex(OCIO_NAMESPACE::SEARCH_REFERENCE_SPACE_ALL, OCIO_NAMESPACE::COLORSPACE_ACTIVE, i);
+				OCIO_NAMESPACE::ConstColorSpaceRcPtr LibColorSpace = CachedNativeConfig->Get()->getColorSpace(ColorSpaceName);
 				if (!LibColorSpace)
 				{
 					continue;
@@ -365,10 +278,10 @@ TSharedRef<SWidget> FOpenColorIOColorSpaceCustomization::HandleSourceComboButton
 
 void FOpenColorIODisplayViewCustomization::PopulateViewSubMenu(FMenuBuilder& InMenuBuilder, FOpenColorIODisplayView InDisplayView)
 {
-	const int32 ViewCount = CachedConfigFile->getNumViews(TCHAR_TO_ANSI(*InDisplayView.Display));
+	const int32 ViewCount = CachedNativeConfig->Get()->getNumViews(TCHAR_TO_ANSI(*InDisplayView.Display));
 	for (int32 i = 0; i < ViewCount; ++i)
 	{
-		const char* ViewName = CachedConfigFile->getView(TCHAR_TO_ANSI(*InDisplayView.Display), i);
+		const char* ViewName = CachedNativeConfig->Get()->getView(TCHAR_TO_ANSI(*InDisplayView.Display), i);
 		InDisplayView.View = StringCast<TCHAR>(ViewName).Get();
 		
 		AddMenuEntry(InMenuBuilder, InDisplayView);
@@ -412,7 +325,7 @@ void FOpenColorIODisplayViewCustomization::AddMenuEntry(FMenuBuilder& InMenuBuil
 
 TSharedRef<SWidget> FOpenColorIODisplayViewCustomization::HandleSourceComboButtonMenuContent()
 {
-	bool bValidConfiguration = CheckValidConfiguration();
+	CachedNativeConfig = GetNativeConfig(ConfigurationObjectProperty);
 
 	// Generate menu
 	FMenuBuilder MenuBuilder(true, nullptr);
@@ -420,12 +333,12 @@ TSharedRef<SWidget> FOpenColorIODisplayViewCustomization::HandleSourceComboButto
 
 	MenuBuilder.BeginSection("AllDisplayViews", LOCTEXT("AllDisplayViewsSection", "Display - View"));
 	{
-		if (bValidConfiguration)
+		if (CachedNativeConfig)
 		{
-			const int32 DisplayCount = CachedConfigFile->getNumDisplays();
+			const int32 DisplayCount = CachedNativeConfig->Get()->getNumDisplays();
 			for (int32 i = 0; i < DisplayCount; ++i)
 			{
-				const char* DisplayName = CachedConfigFile->getDisplay(i);
+				const char* DisplayName = CachedNativeConfig->Get()->getDisplay(i);
 
 				const FOpenColorIODisplayView DisplayViewValue = { StringCast<TCHAR>(DisplayName).Get(), TEXT("<Invalid>") };
 
diff --git a/Engine/Plugins/Compositing/OpenColorIO/Source/OpenColorIOEditor/Private/OpenColorIOColorSpaceCustomization.h b/Engine/Plugins/Compositing/OpenColorIO/Source/OpenColorIOEditor/Private/OpenColorIOColorSpaceCustomization.h
--- a/Engine/Plugins/Compositing/OpenColorIO/Source/OpenColorIOEditor/Private/OpenColorIOColorSpaceCustomization.h
+++ b/Engine/Plugins/Compositing/OpenColorIO/Source/OpenColorIOEditor/Private/OpenColorIOColorSpaceCustomization.h
@@ -7,9 +7,9 @@
 #include "Framework/MultiBox/MultiBoxBuilder.h"
 #include "IPropertyTypeCustomization.h"
 #include "OpenColorIOColorSpace.h"
-#include "OpenColorIO/OpenColorIO.h"
 #include "Widgets/SWidget.h"
 
+class FOpenColorIONativeConfiguration;
 
 /**
  * Implements a details view customization for the FOpenColorIOConfiguration
@@ -17,39 +17,37 @@
 class IPropertyTypeCustomizationOpenColorIO : public IPropertyTypeCustomization
 {
 public:
+	IPropertyTypeCustomizationOpenColorIO(TSharedPtr<IPropertyHandle> InConfigurationObjectProperty);
 
 	/** IPropertyTypeCustomization interface */
 	virtual void CustomizeHeader(TSharedRef<class IPropertyHandle> InPropertyHandle, class FDetailWidgetRow& HeaderRow, IPropertyTypeCustomizationUtils& PropertyTypeCustomizationUtils) override {}
 	virtual void CustomizeChildren(TSharedRef<class IPropertyHandle> InPropertyHandle, class IDetailChildrenBuilder& StructBuilder, IPropertyTypeCustomizationUtils& PropertyTypeCustomizationUtils) override {}
 
-
 protected:
-	TSharedPtr<IPropertyHandle> GetConfigurationFileProperty() const;
-	bool LoadConfigurationFile(const FFilePath& InFilePath);
-	bool CheckValidConfiguration();
+
+	static FOpenColorIONativeConfiguration* GetNativeConfig(const TSharedPtr<IPropertyHandle>& InConfigurationObjectProperty);
 
 protected:
 
 	/** Pointer to the property handle. */
 	TSharedPtr<IPropertyHandle> CachedProperty;
 
-	/** Pointer to the ConfigurationFile property handle. */
-	TSharedPtr<IPropertyHandle> ConfigurationFileProperty;
-
-	/** FilePath of the configuration file that was cached */
-	FFilePath LoadedFilePath;
+	/** Pointer to the configuration object property handle. */
+	TSharedPtr<IPropertyHandle> ConfigurationObjectProperty;
 
-	/** Cached configuration file to populate menus and submenus */
-	OCIO_NAMESPACE::ConstConfigRcPtr CachedConfigFile;
+	/** Raw pointer to the native OCIO config. */
+	FOpenColorIONativeConfiguration* CachedNativeConfig;
 };
 
 
 class FOpenColorIOColorSpaceCustomization : public IPropertyTypeCustomizationOpenColorIO
 {
 public:
-	static TSharedRef<IPropertyTypeCustomization> MakeInstance()
+	using IPropertyTypeCustomizationOpenColorIO::IPropertyTypeCustomizationOpenColorIO;
+
+	static TSharedRef<IPropertyTypeCustomization> MakeInstance(TSharedPtr<IPropertyHandle> InConfigurationObjectProperty)
 	{
-		return MakeShareable(new FOpenColorIOColorSpaceCustomization);
+		return MakeShared<FOpenColorIOColorSpaceCustomization>(MoveTemp(InConfigurationObjectProperty));
 	}
 
 	virtual void CustomizeHeader(TSharedRef<class IPropertyHandle> InPropertyHandle, class FDetailWidgetRow& HeaderRow, IPropertyTypeCustomizationUtils& PropertyTypeCustomizationUtils) override;
@@ -62,15 +60,15 @@ private:
 	TSharedRef<SWidget> HandleSourceComboButtonMenuContent();
 };
 
-/**
- * Implements a details view customization for the FOpenColorIOConfiguration
- */
+
 class FOpenColorIODisplayViewCustomization : public IPropertyTypeCustomizationOpenColorIO
 {
 public:
-	static TSharedRef<IPropertyTypeCustomization> MakeInstance()
+	using IPropertyTypeCustomizationOpenColorIO::IPropertyTypeCustomizationOpenColorIO;
+
+	static TSharedRef<IPropertyTypeCustomization> MakeInstance(TSharedPtr<IPropertyHandle> InConfigurationObjectProperty)
 	{
-		return MakeShareable(new FOpenColorIODisplayViewCustomization);
+		return MakeShared<FOpenColorIODisplayViewCustomization>(MoveTemp(InConfigurationObjectProperty));
 	}
 
 	virtual void CustomizeHeader(TSharedRef<class IPropertyHandle> InPropertyHandle, class FDetailWidgetRow& HeaderRow, IPropertyTypeCustomizationUtils& PropertyTypeCustomizationUtils) override;
diff --git a/Engine/Plugins/Compositing/OpenColorIO/Source/OpenColorIOEditor/Private/OpenColorIOConfigurationCustomization.cpp b/Engine/Plugins/Compositing/OpenColorIO/Source/OpenColorIOEditor/Private/OpenColorIOConfigurationCustomization.cpp
--- a/Engine/Plugins/Compositing/OpenColorIO/Source/OpenColorIOEditor/Private/OpenColorIOConfigurationCustomization.cpp
+++ b/Engine/Plugins/Compositing/OpenColorIO/Source/OpenColorIOEditor/Private/OpenColorIOConfigurationCustomization.cpp
@@ -52,13 +52,15 @@ void FOpenColorIOConfigurationCustomization::CustomizeDetails(IDetailLayoutBuild
 			})
 		];
 
+	TSharedPtr<IPropertyHandle> ConfigurationObjectHandle = ConfigurationFileProperty->GetParentHandle();
+
 	DetailBuilder.RegisterInstancedCustomPropertyTypeLayout(
 			FOpenColorIOColorSpace::StaticStruct()->GetFName(),
-			FOnGetPropertyTypeCustomizationInstance::CreateLambda([] { return MakeShared<FOpenColorIOColorSpaceCustomization>(); }));
+			FOnGetPropertyTypeCustomizationInstance::CreateLambda([ConfigurationObjectHandle] { return FOpenColorIOColorSpaceCustomization::MakeInstance(ConfigurationObjectHandle); }));
 
 	DetailBuilder.RegisterInstancedCustomPropertyTypeLayout(
 		FOpenColorIODisplayView::StaticStruct()->GetFName(),
-		FOnGetPropertyTypeCustomizationInstance::CreateLambda([] { return MakeShared<FOpenColorIODisplayViewCustomization>(); }));
+		FOnGetPropertyTypeCustomizationInstance::CreateLambda([ConfigurationObjectHandle] { return FOpenColorIODisplayViewCustomization::MakeInstance(ConfigurationObjectHandle); }));
 }
 
 #undef LOCTEXT_NAMESPACE
