Changing OpenColorIO to use CreateAndInitSingleView instead of making the view itself.

diff --git a/Engine/Plugins/Compositing/OpenColorIO/Source/OpenColorIO/Private/OpenColorIORendering.cpp b/Engine/Plugins/Compositing/OpenColorIO/Source/OpenColorIO/Private/OpenColorIORendering.cpp
--- a/Engine/Plugins/Compositing/OpenColorIO/Source/OpenColorIO/Private/OpenColorIORendering.cpp
+++ b/Engine/Plugins/Compositing/OpenColorIO/Source/OpenColorIO/Private/OpenColorIORendering.cpp
@@ -5,6 +5,7 @@
 #include "Engine/RendererSettings.h"
 #include "Engine/TextureRenderTarget2D.h"
 #include "Engine/World.h"
+#include "EngineModule.h"
 #include "GlobalShader.h"
 #include "Logging/LogMacros.h"
 #include "Logging/MessageLog.h"
@@ -18,6 +19,7 @@
 #include "SceneInterface.h"
 #include "ScreenPass.h"
 #include "TextureResource.h"
+#include "ScenePrivate.h"
 
 namespace {
 	FViewInfo CreateDummyViewInfo(const FIntRect& InViewRect)
