OCIO: Additional null check in case the native config failed to load.

#rnx

diff --git a/Engine/Plugins/Compositing/OpenColorIO/Source/OpenColorIOEditor/Private/OpenColorIOColorSpaceCustomization.cpp b/Engine/Plugins/Compositing/OpenColorIO/Source/OpenColorIOEditor/Private/OpenColorIOColorSpaceCustomization.cpp
--- a/Engine/Plugins/Compositing/OpenColorIO/Source/OpenColorIOEditor/Private/OpenColorIOColorSpaceCustomization.cpp
+++ b/Engine/Plugins/Compositing/OpenColorIO/Source/OpenColorIOEditor/Private/OpenColorIOColorSpaceCustomization.cpp
@@ -238,7 +238,7 @@ TSharedRef<SWidget> FOpenColorIOColorSpaceCustomization::HandleSourceComboButton
 
 	MenuBuilder.BeginSection("AllColorSpaces", LOCTEXT("AllColorSpacesSection", "ColorSpaces"));
 	{
-		if (CachedNativeConfig)
+		if (CachedNativeConfig && CachedNativeConfig->Get() != nullptr)
 		{
 			const int32 ColorSpaceCount = CachedNativeConfig->Get()->getNumColorSpaces(OCIO_NAMESPACE::SEARCH_REFERENCE_SPACE_ALL, OCIO_NAMESPACE::COLORSPACE_ACTIVE);
 			for (int32 i = 0; i < ColorSpaceCount; ++i)
@@ -333,7 +333,7 @@ TSharedRef<SWidget> FOpenColorIODisplayViewCustomization::HandleSourceComboButto
 
 	MenuBuilder.BeginSection("AllDisplayViews", LOCTEXT("AllDisplayViewsSection", "Display - View"));
 	{
-		if (CachedNativeConfig)
+		if (CachedNativeConfig && CachedNativeConfig->Get() != nullptr)
 		{
 			const int32 DisplayCount = CachedNativeConfig->Get()->getNumDisplays();
 			for (int32 i = 0; i < DisplayCount; ++i)
