Changing OpenColorIO to use CreateAndInitSingleView instead of making the view itself.

diff --git a/Engine/Plugins/Compositing/OpenColorIO/Source/OpenColorIO/Private/OpenColorIORendering.cpp b/Engine/Plugins/Compositing/OpenColorIO/Source/OpenColorIO/Private/OpenColorIORendering.cpp
--- a/Engine/Plugins/Compositing/OpenColorIO/Source/OpenColorIO/Private/OpenColorIORendering.cpp
+++ b/Engine/Plugins/Compositing/OpenColorIO/Source/OpenColorIO/Private/OpenColorIORendering.cpp
@@ -21,23 +21,6 @@
 #include "TextureResource.h"
 #include "ScenePrivate.h"
 
-namespace {
-	FViewInfo CreateDummyViewInfo(const FIntRect& InViewRect)
-	{
-		FSceneViewFamily ViewFamily(FSceneViewFamily::ConstructionValues(nullptr, nullptr, FEngineShowFlags(ESFIM_Game))
-			.SetTime(FGameTime())
-			.SetGammaCorrection(1.0f));
-		FSceneViewInitOptions ViewInitOptions;
-		ViewInitOptions.ViewFamily = &ViewFamily;
-		ViewInitOptions.SetViewRectangle(InViewRect);
-		ViewInitOptions.ViewOrigin = FVector::ZeroVector;
-		ViewInitOptions.ViewRotationMatrix = FMatrix::Identity;
-		ViewInitOptions.ProjectionMatrix = FMatrix::Identity;
-
-		return FViewInfo(ViewInitOptions);
-	}
-}
-
 // static
 void FOpenColorIORendering::AddPass_RenderThread(
 	FRDGBuilder& GraphBuilder,
@@ -137,9 +120,24 @@ bool FOpenColorIORendering::ApplyColorTransform(UWorld* InWorld, const FOpenColo
 			FIntPoint  OutputResolution = FIntPoint(OutputResource->GetSizeX(), OutputResource->GetSizeY());
 			FScreenPassRenderTarget Output = FScreenPassRenderTarget(OutputTexture, FIntRect(FIntPoint::ZeroValue, OutputResolution), ERenderTargetLoadAction::EClear);
 
+
+			FSceneViewFamily ViewFamily(FSceneViewFamily::ConstructionValues(nullptr, nullptr, FEngineShowFlags(ESFIM_Game))
+				.SetTime(FGameTime())
+				.SetGammaCorrection(1.0f));
+
+			FSceneViewInitOptions ViewInitOptions;
+			ViewInitOptions.ViewFamily = &ViewFamily;
+			ViewInitOptions.SetViewRectangle(Output.ViewRect);
+			ViewInitOptions.ViewOrigin = FVector::ZeroVector;
+			ViewInitOptions.ViewRotationMatrix = FMatrix::Identity;
+			ViewInitOptions.ProjectionMatrix = FMatrix::Identity;
+
+			GetRendererModule().CreateAndInitSingleView(RHICmdList, &ViewFamily, &ViewInitOptions);
+			const FViewInfo& View = *(const FViewInfo*)ViewFamily.Views[0];
+
 			AddPass_RenderThread(
 				GraphBuilder,
-				CreateDummyViewInfo(Output.ViewRect),
+				View,
 				FScreenPassTexture(InputTexture),
 				Output,
 				FOpenColorIORenderPassResources{ ShaderResource, TextureResources},
