OCIO: Adds support for automatic conversions from/to the engine working color space.

We rely on the presence of the interchange color space (which is now required in configs as of v2.2) to enable this option in the configuration. When the working color space is selected as a transform source, we insert an additional matrix conversion from working to interchange before applying the normal color transform, and vice-versa for destinations.

diff --git a/Engine/Plugins/Compositing/OpenColorIO/Shaders/Private/OpenColorIOShader.usf b/Engine/Plugins/Compositing/OpenColorIO/Shaders/Private/OpenColorIOShader.usf
--- a/Engine/Plugins/Compositing/OpenColorIO/Shaders/Private/OpenColorIOShader.usf
+++ b/Engine/Plugins/Compositing/OpenColorIO/Shaders/Private/OpenColorIOShader.usf
@@ -10,8 +10,9 @@
 =============================================================================*/
 Texture2D		InputTexture;
 SamplerState	InputTextureSampler;
-
-float Gamma;
+float4x4		WorkingColorSpaceToInterchange;
+float4x4		InterchangeToWorkingColorSpace;
+float			Gamma;
 
 /*=============================================================================
  * OCIO generated start
@@ -36,7 +37,17 @@ void MainPS(in noperspective float4 InUVAndScreenPos : TEXCOORD0,
 	out float4 OutColor : SV_Target0)
 {
 	float4 InputColor = InputTexture.Sample(InputTextureSampler, InUVAndScreenPos.xy);
+	
+#if WORKING_COLOR_SPACE_TRANSFORM_TYPE == 1 // EOpenColorIOWorkingColorSpaceTransform::Source
+	InputColor.xyz = mul((float3x3)WorkingColorSpaceToInterchange, InputColor.xyz);
+#endif
 	InputColor.xyz = max(0.0, InputColor.xyz);
 	InputColor.xyz = pow(InputColor.xyz, Gamma);
+
+	// Apply main transform
 	OutColor = float4(OCIOConvert(InputColor).xyz, InputColor.w);
+
+#if WORKING_COLOR_SPACE_TRANSFORM_TYPE == 2 // EOpenColorIOWorkingColorSpaceTransform::Destination
+	OutColor.xyz = mul((float3x3)InterchangeToWorkingColorSpace, OutColor.xyz);
+#endif
 }
\ No newline at end of file
diff --git a/Engine/Plugins/Compositing/OpenColorIO/Source/OpenColorIO/OpenColorIO.Build.cs b/Engine/Plugins/Compositing/OpenColorIO/Source/OpenColorIO/OpenColorIO.Build.cs
--- a/Engine/Plugins/Compositing/OpenColorIO/Source/OpenColorIO/OpenColorIO.Build.cs
+++ b/Engine/Plugins/Compositing/OpenColorIO/Source/OpenColorIO/OpenColorIO.Build.cs
@@ -25,6 +25,7 @@ namespace UnrealBuildTool.Rules
 					"Renderer",
 					"Slate",
 					"SlateCore",
+					"ColorManagement"
 				});
 
 			PrivateIncludePathModuleNames.AddRange(
diff --git a/Engine/Plugins/Compositing/OpenColorIO/Source/OpenColorIO/Private/OpenColorIOColorTransform.cpp b/Engine/Plugins/Compositing/OpenColorIO/Source/OpenColorIO/Private/OpenColorIOColorTransform.cpp
--- a/Engine/Plugins/Compositing/OpenColorIO/Source/OpenColorIO/Private/OpenColorIOColorTransform.cpp
+++ b/Engine/Plugins/Compositing/OpenColorIO/Source/OpenColorIO/Private/OpenColorIOColorTransform.cpp
@@ -184,6 +184,15 @@ bool UOpenColorIOColorTransform::Initialize(UOpenColorIOConfiguration* InOwner,
 	ConfigurationOwner = InOwner;
 	ContextKeyValues = InContextKeyValues;
 	
+	if (InSourceColorSpace == UOpenColorIOConfiguration::WorkingColorSpaceName)
+	{
+		WorkingColorSpaceTransformType = EOpenColorIOWorkingColorSpaceTransform::Source;
+	}
+	else if(InDestinationColorSpace == UOpenColorIOConfiguration::WorkingColorSpaceName)
+	{
+		WorkingColorSpaceTransformType = EOpenColorIOWorkingColorSpaceTransform::Destination;
+	}
+	
 	return GenerateColorTransformData(InSourceColorSpace, InDestinationColorSpace);
 }
 
@@ -193,6 +202,18 @@ bool UOpenColorIOColorTransform::Initialize(UOpenColorIOConfiguration* InOwner,
 	ConfigurationOwner = InOwner;
 	ContextKeyValues = InContextKeyValues;
 
+	if (InSourceColorSpace == UOpenColorIOConfiguration::WorkingColorSpaceName)
+	{
+		if (InDirection == EOpenColorIOViewTransformDirection::Forward)
+		{
+			WorkingColorSpaceTransformType = EOpenColorIOWorkingColorSpaceTransform::Source;
+		}
+		else
+		{
+			WorkingColorSpaceTransformType = EOpenColorIOWorkingColorSpaceTransform::Destination;
+		}
+	}
+
 	return GenerateColorTransformData(InSourceColorSpace, InDisplay, InView, InDirection);
 }
 
@@ -220,7 +241,7 @@ void UOpenColorIOColorTransform::CacheResourceShadersForCooking(EShaderPlatform
 #else
 	FName AssetPath;
 #endif
-	NewResource->SetupResource((ERHIFeatureLevel::Type)TargetFeatureLevel, InShaderHash, InShaderCode, InRawConfigHash, GetTransformFriendlyName(), AssetPath);
+	NewResource->SetupResource((ERHIFeatureLevel::Type)TargetFeatureLevel, InShaderHash, InShaderCode, InRawConfigHash, GetTransformFriendlyName(), AssetPath, WorkingColorSpaceTransformType);
 
 	const bool bApplyCompletedShaderMap = false;
 	const bool bIsCooking = true;
@@ -411,7 +432,7 @@ void UOpenColorIOColorTransform::CacheResourceShadersForRendering(bool bRegenera
 #else
 				FName AssetPath;
 #endif
-				TransformResource->SetupResource(CacheFeatureLevel, ShaderCodeHash, ShaderCode, RawConfigHash, GetTransformFriendlyName(), AssetPath);
+				TransformResource->SetupResource(CacheFeatureLevel, ShaderCodeHash, ShaderCode, RawConfigHash, GetTransformFriendlyName(), AssetPath, WorkingColorSpaceTransformType);
 
 				const bool bApplyCompletedShaderMap = true;
 
diff --git a/Engine/Plugins/Compositing/OpenColorIO/Source/OpenColorIO/Private/OpenColorIOConfiguration.cpp b/Engine/Plugins/Compositing/OpenColorIO/Source/OpenColorIO/Private/OpenColorIOConfiguration.cpp
--- a/Engine/Plugins/Compositing/OpenColorIO/Source/OpenColorIO/Private/OpenColorIOConfiguration.cpp
+++ b/Engine/Plugins/Compositing/OpenColorIO/Source/OpenColorIO/Private/OpenColorIOConfiguration.cpp
@@ -601,7 +601,7 @@ void UOpenColorIOConfiguration::LoadConfiguration()
 			FString ConfigurationFilePath = ConfigurationFile.FilePath;
 			if (ConfigurationFilePath.StartsWith(TEXT("ocio://")))
 			{
-				NativeConfig->Set(OCIO_NAMESPACE::Config::CreateFromFile(StringCast<ANSICHAR>(*ConfigurationFilePath).Get()));
+				NativeConfig->Set(OCIO_NAMESPACE::Config::CreateFromFile(StringCast<ANSICHAR>(*ConfigurationFilePath).Get()), WorkingColorSpaceName);
 				UE_LOG(LogOpenColorIO, Verbose, TEXT("Loaded built-in OCIO configuration file %s"), *ConfigurationFilePath);
 			}
 			else
@@ -626,7 +626,7 @@ void UOpenColorIOConfiguration::LoadConfiguration()
 				if (NewConfig)
 				{
 					UE_LOG(LogOpenColorIO, Verbose, TEXT("Loaded OCIO configuration file %s"), *FullPath);
-					NativeConfig->Set(NewConfig);
+					NativeConfig->Set(NewConfig, WorkingColorSpaceName);
 					
 					StartDirectoryWatch(FullPath);
 				}
diff --git a/Engine/Plugins/Compositing/OpenColorIO/Source/OpenColorIO/Private/OpenColorIONativeConfiguration.h b/Engine/Plugins/Compositing/OpenColorIO/Source/OpenColorIO/Private/OpenColorIONativeConfiguration.h
--- a/Engine/Plugins/Compositing/OpenColorIO/Source/OpenColorIO/Private/OpenColorIONativeConfiguration.h
+++ b/Engine/Plugins/Compositing/OpenColorIO/Source/OpenColorIO/Private/OpenColorIONativeConfiguration.h
@@ -15,9 +15,36 @@ public:
 		return Config;
 	}
 
-	void Set(OCIO_NAMESPACE::ConstConfigRcPtr InConfig)
+	void Set(OCIO_NAMESPACE::ConstConfigRcPtr InConfig, FStringView InDefaultWorkingColorSpaceName = FStringView())
 	{
-		Config = InConfig;
+		using namespace OCIO_NAMESPACE;
+
+		if (InConfig == nullptr)
+		{
+			Config.reset();
+		}
+		else
+		{	
+			ConstColorSpaceRcPtr InterchangeCS = InConfig->getColorSpace(InConfig->getCanonicalName("aces_interchange"/*OCIO_NAMESPACE::ROLE_INTERCHANGE_SCENE*/));
+
+			// When the aces interchange color space is present, we add the working color space as an additional option.
+			if (InterchangeCS != nullptr && InConfig->getColorSpace(StringCast<ANSICHAR>(InDefaultWorkingColorSpaceName.GetData()).Get()) == nullptr)
+			{
+				ColorSpaceRcPtr WorkingCS = InterchangeCS->createEditableCopy();
+				WorkingCS->setName(StringCast<ANSICHAR>(InDefaultWorkingColorSpaceName.GetData()).Get());
+				WorkingCS->setFamily("UE");
+				WorkingCS->clearAliases();
+
+				ConfigRcPtr ConfigCopy = InConfig->createEditableCopy();
+				ConfigCopy->addColorSpace(WorkingCS);
+
+				Config = ConfigCopy;
+			}
+			else
+			{
+				Config = InConfig;
+			}
+		}
 	}
 #endif
 
diff --git a/Engine/Plugins/Compositing/OpenColorIO/Source/OpenColorIO/Private/OpenColorIORendering.cpp b/Engine/Plugins/Compositing/OpenColorIO/Source/OpenColorIO/Private/OpenColorIORendering.cpp
--- a/Engine/Plugins/Compositing/OpenColorIO/Source/OpenColorIO/Private/OpenColorIORendering.cpp
+++ b/Engine/Plugins/Compositing/OpenColorIO/Source/OpenColorIO/Private/OpenColorIORendering.cpp
@@ -2,6 +2,7 @@
 
 #include "OpenColorIORendering.h"
 
+#include "ColorSpace.h"
 #include "Engine/RendererSettings.h"
 #include "Engine/TextureRenderTarget2D.h"
 #include "Engine/World.h"
@@ -21,6 +22,24 @@
 #include "TextureResource.h"
 #include "ScenePrivate.h"
 
+
+namespace {
+	using namespace UE::Color;
+
+	// Static local storage to prevent color space recomputation every frame. This is viable since a WCS change requires a project relaunch.
+	const FMatrix44f& GetWorkingColorSpaceToInterchangeTransform()
+	{
+		static FMatrix44f Transform = Transpose<float>(FColorSpaceTransform(FColorSpace::GetWorking(), FColorSpace(EColorSpace::ACESAP0)));
+		return Transform;
+	}
+
+	const FMatrix44f& GetInterchangeToWorkingColorSpaceTransform()
+	{
+		static FMatrix44f Transform = Transpose<float>(FColorSpaceTransform(FColorSpace(EColorSpace::ACESAP0), FColorSpace::GetWorking()));
+		return Transform;
+	}
+}
+
 // static
 void FOpenColorIORendering::AddPass_RenderThread(
 	FRDGBuilder& GraphBuilder,
@@ -43,6 +62,21 @@ void FOpenColorIORendering::AddPass_RenderThread(
 		Parameters->InputTexture = Input.Texture;
 		Parameters->InputTextureSampler = TStaticSamplerState<>::GetRHI();
 		OpenColorIOBindTextureResources(Parameters, InPassResource.TextureResources);
+
+		// Apply a transform between the working color space and the interchange color space, if necessary.
+		switch (InPassResource.ShaderResource->GetWorkingColorSpaceTransformType())
+		{
+		case EOpenColorIOWorkingColorSpaceTransform::Source:
+			Parameters->WorkingColorSpaceToInterchange = GetWorkingColorSpaceToInterchangeTransform();
+			break;
+		case EOpenColorIOWorkingColorSpaceTransform::Destination:
+			Parameters->InterchangeToWorkingColorSpace = GetInterchangeToWorkingColorSpaceTransform();
+			break;
+
+		default:
+			// do nothing, shader parameter is unused.
+			break;	
+		}
 		Parameters->Gamma = InGamma;
 		Parameters->RenderTargets[0] = Output.GetRenderTargetBinding();
 
diff --git a/Engine/Plugins/Compositing/OpenColorIO/Source/OpenColorIO/Private/OpenColorIOShaderMap.cpp b/Engine/Plugins/Compositing/OpenColorIO/Source/OpenColorIO/Private/OpenColorIOShaderMap.cpp
--- a/Engine/Plugins/Compositing/OpenColorIO/Source/OpenColorIO/Private/OpenColorIOShaderMap.cpp
+++ b/Engine/Plugins/Compositing/OpenColorIO/Source/OpenColorIO/Private/OpenColorIOShaderMap.cpp
@@ -209,6 +209,13 @@ FShader* FOpenColorIOShaderType::FinishCompileShader(
 	return Shader;
 }
 
+void FOpenColorIOShaderType::SetupCompileEnvironment(EShaderPlatform InPlatform, const FOpenColorIOTransformResource* InColorTransform, FShaderCompilerEnvironment& OutEnvironment) const
+{
+	ModifyCompilationEnvironment(FOpenColorIOShaderPermutationParameters(GetFName(), InPlatform, InColorTransform), OutEnvironment);
+
+	OutEnvironment.SetDefine(TEXT("WORKING_COLOR_SPACE_TRANSFORM_TYPE"), (int32)InColorTransform->GetWorkingColorSpaceTransformType());
+}
+
 #endif // WITH_EDITOR
 
 /**
diff --git a/Engine/Plugins/Compositing/OpenColorIO/Source/OpenColorIO/Private/OpenColorIOShared.cpp b/Engine/Plugins/Compositing/OpenColorIO/Source/OpenColorIO/Private/OpenColorIOShared.cpp
--- a/Engine/Plugins/Compositing/OpenColorIO/Source/OpenColorIO/Private/OpenColorIOShared.cpp
+++ b/Engine/Plugins/Compositing/OpenColorIO/Source/OpenColorIO/Private/OpenColorIOShared.cpp
@@ -199,7 +199,7 @@ void FOpenColorIOTransformResource::SerializeShaderMap(FArchive& Ar)
 	}
 }
 
-void FOpenColorIOTransformResource::SetupResource(ERHIFeatureLevel::Type InFeatureLevel, const FString& InShaderCodeHash, const FString& InShadercode, const FString& InRawConfigHash, const FString& InFriendlyName, const FName& InAssetPath)
+void FOpenColorIOTransformResource::SetupResource(ERHIFeatureLevel::Type InFeatureLevel, const FString& InShaderCodeHash, const FString& InShadercode, const FString& InRawConfigHash, const FString& InFriendlyName, const FName& InAssetPath, EOpenColorIOWorkingColorSpaceTransform InWorkingColorSpaceTransformType)
 {
 	// When this happens we assume that shader was cooked and we don't need to do anything. 
 	if (InShaderCodeHash.IsEmpty() && InRawConfigHash.IsEmpty())
@@ -207,15 +207,15 @@ void FOpenColorIOTransformResource::SetupResource(ERHIFeatureLevel::Type InFeatu
 		return;
 	}
 
-	FString FullShaderCodeHash;
-	FullShaderCodeHash.Reserve(InShaderCodeHash.Len());
-	const FString ConcatHashes = InShaderCodeHash + InRawConfigHash;
+
+	const FString ConcatHashes = InShaderCodeHash + InRawConfigHash + FString::FromInt((int32)InWorkingColorSpaceTransformType);
 	FSHAHash FullHash;
 	FSHA1::HashBuffer(TCHAR_TO_ANSI(*(ConcatHashes)), ConcatHashes.Len(), FullHash.Hash);
 
 	ShaderCodeAndConfigHash = FullHash.ToString();
 	ShaderCode = InShadercode;
 	FriendlyName = InFriendlyName.Replace(TEXT("/"), TEXT(""));
+	WorkingColorSpaceTransformType = InWorkingColorSpaceTransformType;
 #if WITH_EDITOR
 	AssetPath = InAssetPath;
 #endif // WITH_EDITOR
diff --git a/Engine/Plugins/Compositing/OpenColorIO/Source/OpenColorIO/Public/OpenColorIOColorTransform.h b/Engine/Plugins/Compositing/OpenColorIO/Source/OpenColorIO/Public/OpenColorIOColorTransform.h
--- a/Engine/Plugins/Compositing/OpenColorIO/Source/OpenColorIO/Public/OpenColorIOColorTransform.h
+++ b/Engine/Plugins/Compositing/OpenColorIO/Source/OpenColorIO/Public/OpenColorIOColorTransform.h
@@ -214,6 +214,9 @@ private:
 	/** Key-value string pairs used to define the processor's context. */
 	TMap<FString, FString> ContextKeyValues;
 
+	/** Enum used to indicate whether the working color space should be used as a source or destination. */
+	EOpenColorIOWorkingColorSpaceTransform WorkingColorSpaceTransformType = EOpenColorIOWorkingColorSpaceTransform::None;
+
 	FRenderCommandFence ReleaseFence;
 
 #if WITH_EDITORONLY_DATA
diff --git a/Engine/Plugins/Compositing/OpenColorIO/Source/OpenColorIO/Public/OpenColorIOConfiguration.h b/Engine/Plugins/Compositing/OpenColorIO/Source/OpenColorIO/Public/OpenColorIOConfiguration.h
--- a/Engine/Plugins/Compositing/OpenColorIO/Source/OpenColorIO/Public/OpenColorIOConfiguration.h
+++ b/Engine/Plugins/Compositing/OpenColorIO/Source/OpenColorIO/Public/OpenColorIOConfiguration.h
@@ -35,6 +35,9 @@ public:
 	UOpenColorIOConfiguration(const FObjectInitializer& ObjectInitializer);
 
 public:
+
+	static constexpr TCHAR WorkingColorSpaceName[] = TEXT("Working Color Space");
+
 	bool IsTransformReady(const FOpenColorIOColorConversionSettings& InSettings);
 	bool GetRenderResources(ERHIFeatureLevel::Type InFeatureLevel, const FOpenColorIOColorConversionSettings& InSettings, FOpenColorIOTransformResource*& OutShaderResource, TSortedMap<int32, FTextureResource*>& OutTextureResources);
 	UE_DEPRECATED(5.1, "GetShaderAndLUTResources is deprecated, please use GetRenderResources instead.")
diff --git a/Engine/Plugins/Compositing/OpenColorIO/Source/OpenColorIO/Public/OpenColorIOShader.h b/Engine/Plugins/Compositing/OpenColorIO/Source/OpenColorIO/Public/OpenColorIOShader.h
--- a/Engine/Plugins/Compositing/OpenColorIO/Source/OpenColorIO/Public/OpenColorIOShader.h
+++ b/Engine/Plugins/Compositing/OpenColorIO/Source/OpenColorIO/Public/OpenColorIOShader.h
@@ -20,6 +20,8 @@ class FRHITexture;
 BEGIN_SHADER_PARAMETER_STRUCT(FOpenColorIOPixelShaderParameters, )
 	SHADER_PARAMETER_RDG_TEXTURE(Texture2D, InputTexture)
 	SHADER_PARAMETER_SAMPLER(SamplerState, InputTextureSampler)
+	SHADER_PARAMETER(FMatrix44f, WorkingColorSpaceToInterchange)
+	SHADER_PARAMETER(FMatrix44f, InterchangeToWorkingColorSpace)
 	SHADER_PARAMETER(float, Gamma)
 	
 	SHADER_PARAMETER_TEXTURE(Texture3D, Ocio_lut3d_0)
diff --git a/Engine/Plugins/Compositing/OpenColorIO/Source/OpenColorIO/Public/OpenColorIOShaderType.h b/Engine/Plugins/Compositing/OpenColorIO/Source/OpenColorIO/Public/OpenColorIOShaderType.h
--- a/Engine/Plugins/Compositing/OpenColorIO/Source/OpenColorIO/Public/OpenColorIOShaderType.h
+++ b/Engine/Plugins/Compositing/OpenColorIO/Source/OpenColorIO/Public/OpenColorIOShaderType.h
@@ -138,9 +138,6 @@ protected:
 	 * @param InPlatform - Platform to compile for.
 	 * @param OutEnvironment - The shader compile environment that the function modifies.
 	 */
-	void SetupCompileEnvironment(EShaderPlatform InPlatform, const FOpenColorIOTransformResource* InColorTransform, FShaderCompilerEnvironment& OutEnvironment) const
-	{
-		ModifyCompilationEnvironment(FOpenColorIOShaderPermutationParameters(GetFName(), InPlatform, InColorTransform), OutEnvironment);
-	}
+	void SetupCompileEnvironment(EShaderPlatform InPlatform, const FOpenColorIOTransformResource* InColorTransform, FShaderCompilerEnvironment& OutEnvironment) const;
 #endif // WITH_EDITOR
 };
diff --git a/Engine/Plugins/Compositing/OpenColorIO/Source/OpenColorIO/Public/OpenColorIOShared.h b/Engine/Plugins/Compositing/OpenColorIO/Source/OpenColorIO/Public/OpenColorIOShared.h
--- a/Engine/Plugins/Compositing/OpenColorIO/Source/OpenColorIO/Public/OpenColorIOShared.h
+++ b/Engine/Plugins/Compositing/OpenColorIO/Source/OpenColorIO/Public/OpenColorIOShared.h
@@ -29,6 +29,14 @@ class FOpenColorIOShaderMap;
 class FOpenColorIOPixelShader;
 class FOpenColorIOShaderMapId;
 
+/** Enum used to indicate whether the working color space should be used as a source or destination. */
+enum class EOpenColorIOWorkingColorSpaceTransform : uint8
+{
+	None = 0,
+	Source = 1,
+	Destination = 2
+};
+
 /** Stores outputs from the color transform compile that need to be saved. */
 class FOpenColorIOCompilationOutput
 {
@@ -318,6 +326,7 @@ public:
 		, FeatureLevel(ERHIFeatureLevel::SM5)
 		, bContainsInlineShaders(false)
 		, bLoadedCookedShaderMapId(false)
+		, WorkingColorSpaceTransformType(EOpenColorIOWorkingColorSpaceTransform::None)
 	{}
 
 	/**
@@ -436,7 +445,15 @@ public:
 	const FString& GetFriendlyName()	const { return FriendlyName; }
 
 
-	void SetupResource(ERHIFeatureLevel::Type InFeatureLevel, const FString& InShaderCodeHash, const FString& InShadercode, const FString& InRawConfigHash, const FString& InFriendlyName, const FName& InAssetPath);
+	void SetupResource(
+		ERHIFeatureLevel::Type InFeatureLevel,
+		const FString& InShaderCodeHash,
+		const FString& InShadercode,
+		const FString& InRawConfigHash,
+		const FString& InFriendlyName,
+		const FName& InAssetPath,
+		EOpenColorIOWorkingColorSpaceTransform InWorkingColorSpaceTransformType = EOpenColorIOWorkingColorSpaceTransform::None
+	);
 
 	void SetCompileErrors(TArray<FString> &InErrors)
 	{
@@ -470,6 +487,7 @@ public:
 
 	
 	bool IsSame(const FOpenColorIOShaderMapId& InId) const;
+	EOpenColorIOWorkingColorSpaceTransform GetWorkingColorSpaceTransformType() const { return WorkingColorSpaceTransformType; }
 protected:
 #if WITH_EDITOR
 	/**
@@ -526,6 +544,7 @@ private:
 	uint32 bContainsInlineShaders : 1;
 	uint32 bLoadedCookedShaderMapId : 1;
 	FOpenColorIOShaderMapId CookedShaderMapId;
+	EOpenColorIOWorkingColorSpaceTransform WorkingColorSpaceTransformType;
 
 #if WITH_EDITOR
 	/**
