Fixing compile errors after merge.

diff --git a/Engine/Plugins/Compositing/OpenColorIO/Source/OpenColorIO/Private/OpenColorIORendering.cpp b/Engine/Plugins/Compositing/OpenColorIO/Source/OpenColorIO/Private/OpenColorIORendering.cpp
--- a/Engine/Plugins/Compositing/OpenColorIO/Source/OpenColorIO/Private/OpenColorIORendering.cpp
+++ b/Engine/Plugins/Compositing/OpenColorIO/Source/OpenColorIO/Private/OpenColorIORendering.cpp
@@ -84,7 +84,8 @@ void FOpenColorIORendering::AddPass_RenderThread(
 	else
 	{
 		// Fallback pass, printing invalid message across the viewport.
-		TShaderMapRef<FOpenColorIOInvalidPixelShader> OCIOInvalidPixelShader(View.ShaderMap);
+		FGlobalShaderMap* ShaderMap = GetGlobalShaderMap(View.GetFeatureLevel());
+		TShaderMapRef<FOpenColorIOInvalidPixelShader> OCIOInvalidPixelShader(ShaderMap);
 		FOpenColorIOInvalidShaderParameters* Parameters = GraphBuilder.AllocParameters<FOpenColorIOInvalidShaderParameters>();
 		Parameters->InputTexture = Input.Texture;
 		Parameters->InputTextureSampler = TStaticSamplerState<>::GetRHI();
@@ -166,7 +167,7 @@ bool FOpenColorIORendering::ApplyColorTransform(UWorld* InWorld, const FOpenColo
 			ViewInitOptions.ProjectionMatrix = FMatrix::Identity;
 
 			GetRendererModule().CreateAndInitSingleView(RHICmdList, &ViewFamily, &ViewInitOptions);
-			const FViewInfo& View = *(const FViewInfo*)ViewFamily.Views[0];
+			const FSceneView& View = *ViewFamily.Views[0];
 
 			AddPass_RenderThread(
 				GraphBuilder,
