Fix private include path in OpenColorIO

#rnx

diff --git a/Engine/Plugins/Compositing/OpenColorIO/Source/OpenColorIO/Private/OpenColorIONativeConfiguration.h b/Engine/Plugins/Compositing/OpenColorIO/Source/OpenColorIO/Private/OpenColorIONativeConfiguration.h
deleted file mode 100644
diff --git a/Engine/Plugins/Compositing/OpenColorIO/Source/OpenColorIO/Public/OpenColorIONativeConfiguration.h b/Engine/Plugins/Compositing/OpenColorIO/Source/OpenColorIO/Public/OpenColorIONativeConfiguration.h
new file mode 100644
--- /dev/null
+++ b/Engine/Plugins/Compositing/OpenColorIO/Source/OpenColorIO/Public/OpenColorIONativeConfiguration.h
@@ -0,0 +1,55 @@
+// Copyright Epic Games, Inc. All Rights Reserved.
+
+#pragma once
+
+#if WITH_EDITOR && WITH_OCIO
+#include "OpenColorIO/OpenColorIO.h"
+#endif
+
+class FOpenColorIONativeConfiguration {
+public:
+
+#if WITH_EDITOR && WITH_OCIO
+	OCIO_NAMESPACE::ConstConfigRcPtr Get() const
+	{
+		return Config;
+	}
+
+	void Set(OCIO_NAMESPACE::ConstConfigRcPtr InConfig, FStringView InDefaultWorkingColorSpaceName = FStringView())
+	{
+		using namespace OCIO_NAMESPACE;
+
+		if (InConfig == nullptr)
+		{
+			Config.reset();
+		}
+		else
+		{	
+			ConstColorSpaceRcPtr InterchangeCS = InConfig->getColorSpace(InConfig->getCanonicalName("aces_interchange"/*OCIO_NAMESPACE::ROLE_INTERCHANGE_SCENE*/));
+
+			// When the aces interchange color space is present, we add the working color space as an additional option.
+			if (InterchangeCS != nullptr && InConfig->getColorSpace(StringCast<ANSICHAR>(InDefaultWorkingColorSpaceName.GetData()).Get()) == nullptr)
+			{
+				ColorSpaceRcPtr WorkingCS = InterchangeCS->createEditableCopy();
+				WorkingCS->setName(StringCast<ANSICHAR>(InDefaultWorkingColorSpaceName.GetData()).Get());
+				WorkingCS->setFamily("UE");
+				WorkingCS->clearAliases();
+
+				ConfigRcPtr ConfigCopy = InConfig->createEditableCopy();
+				ConfigCopy->addColorSpace(WorkingCS);
+
+				Config = ConfigCopy;
+			}
+			else
+			{
+				Config = InConfig;
+			}
+		}
+	}
+#endif
+
+private:
+#if WITH_EDITORONLY_DATA && WITH_OCIO
+	OCIO_NAMESPACE::ConstConfigRcPtr Config;
+#endif //WITH_EDITORONLY_DATA
+};
diff --git a/Engine/Plugins/Compositing/OpenColorIO/Source/OpenColorIOEditor/Private/OpenColorIOColorSpaceCustomization.cpp b/Engine/Plugins/Compositing/OpenColorIO/Source/OpenColorIOEditor/Private/OpenColorIOColorSpaceCustomization.cpp
--- a/Engine/Plugins/Compositing/OpenColorIO/Source/OpenColorIOEditor/Private/OpenColorIOColorSpaceCustomization.cpp
+++ b/Engine/Plugins/Compositing/OpenColorIO/Source/OpenColorIOEditor/Private/OpenColorIOColorSpaceCustomization.cpp
@@ -7,7 +7,7 @@
 #include "IPropertyUtilities.h"
 #include "PropertyHandle.h"
 #include "OpenColorIOConfiguration.h"
-#include "OpenColorIO/Private/OpenColorIONativeConfiguration.h"
+#include "OpenColorIONativeConfiguration.h"
 #include "OpenColorIO/OpenColorIO.h"
 #include "Widgets/Input/SComboButton.h"
 
