Fixing OpenColorIO after ScreenPass changes.

diff --git a/Engine/Plugins/Compositing/OpenColorIO/Source/OpenColorIO/Private/OpenColorIODisplayExtension.cpp b/Engine/Plugins/Compositing/OpenColorIO/Source/OpenColorIO/Private/OpenColorIODisplayExtension.cpp
--- a/Engine/Plugins/Compositing/OpenColorIO/Source/OpenColorIO/Private/OpenColorIODisplayExtension.cpp
+++ b/Engine/Plugins/Compositing/OpenColorIO/Source/OpenColorIO/Private/OpenColorIODisplayExtension.cpp
@@ -127,6 +127,7 @@ FScreenPassTexture FOpenColorIODisplayExtension::PostProcessPassAfterTonemap_Ren
 	FOpenColorIORendering::AddPass_RenderThread(
 		GraphBuilder,
 		View,
+		View.GetFeatureLevel(),
 		SceneColor,
 		Output,
 		CachedResourcesRenderThread,
diff --git a/Engine/Plugins/Compositing/OpenColorIO/Source/OpenColorIO/Private/OpenColorIORendering.cpp b/Engine/Plugins/Compositing/OpenColorIO/Source/OpenColorIO/Private/OpenColorIORendering.cpp
--- a/Engine/Plugins/Compositing/OpenColorIO/Source/OpenColorIO/Private/OpenColorIORendering.cpp
+++ b/Engine/Plugins/Compositing/OpenColorIO/Source/OpenColorIO/Private/OpenColorIORendering.cpp
@@ -42,7 +42,8 @@ namespace {
 // static
 void FOpenColorIORendering::AddPass_RenderThread(
 	FRDGBuilder& GraphBuilder,
-	const FSceneView& View,
+	const FScreenPassViewInfo ViewInfo,
+	ERHIFeatureLevel::Type FeatureLevel,
 	const FScreenPassTexture& Input,
 	const FScreenPassRenderTarget& Output,
 	const FOpenColorIORenderPassResources& InPassResource,
@@ -53,6 +54,9 @@ void FOpenColorIORendering::AddPass_RenderThread(
 	const FScreenPassTextureViewport InputViewport(Input);
 	const FScreenPassTextureViewport OutputViewport(Output);
 
+	FGlobalShaderMap* ShaderMap = GetGlobalShaderMap(FeatureLevel);
+	TShaderMapRef<FScreenPassVS> VertexShader(ShaderMap);
+
 	if (InPassResource.ShaderResource != nullptr)
 	{
 		TShaderRef<FOpenColorIOPixelShader> OCIOPixelShader = InPassResource.ShaderResource->GetShader<FOpenColorIOPixelShader>();
@@ -79,12 +83,11 @@ void FOpenColorIORendering::AddPass_RenderThread(
 		Parameters->Gamma = InGamma;
 		Parameters->RenderTargets[0] = Output.GetRenderTargetBinding();
 
-		AddDrawScreenPass(GraphBuilder, RDG_EVENT_NAME("OpenColorIOPass"), View, OutputViewport, InputViewport, OCIOPixelShader, Parameters);
+		AddDrawScreenPass(GraphBuilder, RDG_EVENT_NAME("OpenColorIOPass"), ViewInfo, OutputViewport, InputViewport, VertexShader, OCIOPixelShader, Parameters);
 	}
 	else
 	{
 		// Fallback pass, printing invalid message across the viewport.
-		FGlobalShaderMap* ShaderMap = GetGlobalShaderMap(View.GetFeatureLevel());
 		TShaderMapRef<FOpenColorIOInvalidPixelShader> OCIOInvalidPixelShader(ShaderMap);
 		FOpenColorIOInvalidShaderParameters* Parameters = GraphBuilder.AllocParameters<FOpenColorIOInvalidShaderParameters>();
 		Parameters->InputTexture = Input.Texture;
@@ -92,7 +95,7 @@ void FOpenColorIORendering::AddPass_RenderThread(
 		Parameters->MiniFontTexture = OpenColorIOGetMiniFontTexture();
 		Parameters->RenderTargets[0] = Output.GetRenderTargetBinding();
 
-		AddDrawScreenPass(GraphBuilder, RDG_EVENT_NAME("OpenColorIOInvalidPass"), View, OutputViewport, InputViewport, OCIOInvalidPixelShader, Parameters);
+		AddDrawScreenPass(GraphBuilder, RDG_EVENT_NAME("OpenColorIOInvalidPass"), ViewInfo, OutputViewport, InputViewport, VertexShader, OCIOInvalidPixelShader, Parameters);
 	}
 }
 
@@ -145,7 +148,7 @@ bool FOpenColorIORendering::ApplyColorTransform(UWorld* InWorld, const FOpenColo
 	}
 	
 	ENQUEUE_RENDER_COMMAND(ProcessColorSpaceTransform)(
-		[InputResource, OutputResource, ShaderResource, TextureResources = MoveTemp(TransformTextureResources)](FRHICommandListImmediate& RHICmdList)
+		[FeatureLevel, InputResource, OutputResource, ShaderResource, TextureResources = MoveTemp(TransformTextureResources)](FRHICommandListImmediate& RHICmdList)
 		{
 			FRDGBuilder GraphBuilder(RHICmdList);
 
@@ -154,24 +157,10 @@ bool FOpenColorIORendering::ApplyColorTransform(UWorld* InWorld, const FOpenColo
 			FIntPoint  OutputResolution = FIntPoint(OutputResource->GetSizeX(), OutputResource->GetSizeY());
 			FScreenPassRenderTarget Output = FScreenPassRenderTarget(OutputTexture, FIntRect(FIntPoint::ZeroValue, OutputResolution), ERenderTargetLoadAction::EClear);
 
-
-			FSceneViewFamily ViewFamily(FSceneViewFamily::ConstructionValues(nullptr, nullptr, FEngineShowFlags(ESFIM_Game))
-				.SetTime(FGameTime())
-				.SetGammaCorrection(1.0f));
-
-			FSceneViewInitOptions ViewInitOptions;
-			ViewInitOptions.ViewFamily = &ViewFamily;
-			ViewInitOptions.SetViewRectangle(Output.ViewRect);
-			ViewInitOptions.ViewOrigin = FVector::ZeroVector;
-			ViewInitOptions.ViewRotationMatrix = FMatrix::Identity;
-			ViewInitOptions.ProjectionMatrix = FMatrix::Identity;
-
-			GetRendererModule().CreateAndInitSingleView(RHICmdList, &ViewFamily, &ViewInitOptions);
-			const FSceneView& View = *ViewFamily.Views[0];
-
 			AddPass_RenderThread(
 				GraphBuilder,
-				View,
+				FScreenPassViewInfo(),
+				FeatureLevel,
 				FScreenPassTexture(InputTexture),
 				Output,
 				FOpenColorIORenderPassResources{ ShaderResource, TextureResources},
diff --git a/Engine/Plugins/Compositing/OpenColorIO/Source/OpenColorIO/Public/OpenColorIORendering.h b/Engine/Plugins/Compositing/OpenColorIO/Source/OpenColorIO/Public/OpenColorIORendering.h
--- a/Engine/Plugins/Compositing/OpenColorIO/Source/OpenColorIO/Public/OpenColorIORendering.h
+++ b/Engine/Plugins/Compositing/OpenColorIO/Source/OpenColorIO/Public/OpenColorIORendering.h
@@ -9,6 +9,7 @@ class FOpenColorIOTransformResource;
 class FRDGBuilder;
 struct FScreenPassRenderTarget;
 struct FScreenPassTexture;
+struct FScreenPassViewInfo;
 class FTextureResource;
 class FSceneView;
 class UTexture;
@@ -50,7 +51,8 @@ public:
 	 */
 	static void AddPass_RenderThread(
 		FRDGBuilder& GraphBuilder,
-		const FSceneView& View,
+		FScreenPassViewInfo ViewInfo,
+		ERHIFeatureLevel::Type FeatureLevel,
 		const FScreenPassTexture& Input,
 		const FScreenPassRenderTarget& Output,
 		const FOpenColorIORenderPassResources& InPassInfo,
