OCIO: Base implementation of CPU-side color transforms.

Support for working color space conversions to come next.

diff --git a/Engine/Plugins/Compositing/OpenColorIO/Source/OpenColorIO/OpenColorIO.Build.cs b/Engine/Plugins/Compositing/OpenColorIO/Source/OpenColorIO/OpenColorIO.Build.cs
--- a/Engine/Plugins/Compositing/OpenColorIO/Source/OpenColorIO/OpenColorIO.Build.cs
+++ b/Engine/Plugins/Compositing/OpenColorIO/Source/OpenColorIO/OpenColorIO.Build.cs
@@ -23,6 +23,7 @@ namespace UnrealBuildTool.Rules
 					"RHI",
 					"RenderCore",
 					"Renderer",
+					"ImageCore",
 					"Slate",
 					"SlateCore",
 					"ColorManagement"
diff --git a/Engine/Plugins/Compositing/OpenColorIO/Source/OpenColorIO/Private/OpenColorIOColorTransform.cpp b/Engine/Plugins/Compositing/OpenColorIO/Source/OpenColorIO/Private/OpenColorIOColorTransform.cpp
--- a/Engine/Plugins/Compositing/OpenColorIO/Source/OpenColorIO/Private/OpenColorIOColorTransform.cpp
+++ b/Engine/Plugins/Compositing/OpenColorIO/Source/OpenColorIO/Private/OpenColorIOColorTransform.cpp
@@ -34,7 +34,6 @@ namespace {
 #if WITH_EDITOR && WITH_OCIO
 	/*
 	 * Get the processor optimization flag.
-	 * @todo: Remove "no dynamic properties" once this path is enabled.
 	 */
 	OCIO_NAMESPACE::OptimizationFlags GetProcessorOptimization()
 	{
@@ -43,36 +42,90 @@ namespace {
 		return static_cast<OptimizationFlags>(OptimizationFlags::OPTIMIZATION_DEFAULT | OptimizationFlags::OPTIMIZATION_NO_DYNAMIC_PROPERTIES);
 	}
 
-	OCIO_NAMESPACE::ConstProcessorRcPtr GetTransformProcessor(UOpenColorIOColorTransform* InTransform, const OCIO_NAMESPACE::ConstConfigRcPtr& InConfig, const TMap<FString, FString>& InContext)
+	/*
+	 * Get the OpenColorIO processor using the transform's config.
+	 * NOTE: By default, the library automatically caches and reuses processors.
+	 */
+	OCIO_NAMESPACE::ConstProcessorRcPtr GetTransformProcessor(const UOpenColorIOColorTransform* InTransform)
 	{
-		EOpenColorIOViewTransformDirection DisplayViewDirection;
-		OCIO_NAMESPACE::ContextRcPtr Context = InConfig->getCurrentContext()->createEditableCopy();
+		check(InTransform);
+		check(InTransform->ConfigurationOwner);
 
-		for (const TPair<FString, FString>& KeyValue : InContext)
+		OCIO_NAMESPACE::ConstConfigRcPtr Config = InTransform->ConfigurationOwner->GetNativeConfig_Internal()->Get();
+		if (Config)
 		{
-			Context->setStringVar(TCHAR_TO_ANSI(*KeyValue.Key), TCHAR_TO_ANSI(*KeyValue.Value));
-		}
+			OCIO_NAMESPACE::ContextRcPtr Context = Config->getCurrentContext()->createEditableCopy();
 
-		if (InTransform->GetDisplayViewDirection(DisplayViewDirection))
-		{
-			OCIO_NAMESPACE::TransformDirection OcioDirection = static_cast<OCIO_NAMESPACE::TransformDirection>(DisplayViewDirection);
-
-			return InConfig->getProcessor(
-				Context,
-				StringCast<ANSICHAR>(*InTransform->SourceColorSpace).Get(),
-				StringCast<ANSICHAR>(*InTransform->Display).Get(),
-				StringCast<ANSICHAR>(*InTransform->View).Get(),
-				OcioDirection);
+			for (const TPair<FString, FString>& KeyValue : InTransform->GetContextKeyValues())
+			{
+				Context->setStringVar(TCHAR_TO_ANSI(*KeyValue.Key), TCHAR_TO_ANSI(*KeyValue.Value));
+			}
+
+			EOpenColorIOViewTransformDirection DisplayViewDirection;
+			if (InTransform->GetDisplayViewDirection(DisplayViewDirection))
+			{
+				OCIO_NAMESPACE::TransformDirection OcioDirection = static_cast<OCIO_NAMESPACE::TransformDirection>(DisplayViewDirection);
+
+				return Config->getProcessor(
+					Context,
+					StringCast<ANSICHAR>(*InTransform->SourceColorSpace).Get(),
+					StringCast<ANSICHAR>(*InTransform->Display).Get(),
+					StringCast<ANSICHAR>(*InTransform->View).Get(),
+					OcioDirection);
+			}
+			else
+			{
+				return Config->getProcessor(
+					Context,
+					StringCast<ANSICHAR>(*InTransform->SourceColorSpace).Get(),
+					StringCast<ANSICHAR>(*InTransform->DestinationColorSpace).Get()
+				);
+			}
 		}
-		else
+
+		return nullptr;
+	}
+
+	TUniquePtr<OCIO_NAMESPACE::PackedImageDesc> GetImageDesc(const FImageView& InImage)
+	{
+		OCIO_NAMESPACE::ChannelOrdering Ordering;
+		OCIO_NAMESPACE::BitDepth BitDepth;
+
+		switch (InImage.Format)
 		{
-			return InConfig->getProcessor(
-				Context,
-				StringCast<ANSICHAR>(*InTransform->SourceColorSpace).Get(),
-				StringCast<ANSICHAR>(*InTransform->DestinationColorSpace).Get()
-			);
+		case ERawImageFormat::BGRA8:
+			Ordering = OCIO_NAMESPACE::CHANNEL_ORDERING_BGRA;
+			BitDepth = OCIO_NAMESPACE::BIT_DEPTH_UINT8;
+			break;
+		case ERawImageFormat::RGBA16:
+			Ordering = OCIO_NAMESPACE::CHANNEL_ORDERING_RGBA;
+			BitDepth = OCIO_NAMESPACE::BIT_DEPTH_UINT16;
+			break;
+		case ERawImageFormat::RGBA16F:
+			Ordering = OCIO_NAMESPACE::CHANNEL_ORDERING_RGBA;
+			BitDepth = OCIO_NAMESPACE::BIT_DEPTH_F16;
+			break;
+		case ERawImageFormat::RGBA32F:
+			Ordering = OCIO_NAMESPACE::CHANNEL_ORDERING_RGBA;
+			BitDepth = OCIO_NAMESPACE::BIT_DEPTH_F32;
+			break;
+		default:
+
+			UE_LOG(LogOpenColorIO, Log, TEXT("Unsupported texture format."));
+			return nullptr;
 		}
-	}
+
+		return MakeUnique<OCIO_NAMESPACE::PackedImageDesc>(
+			InImage.RawData,
+			static_cast<long>(InImage.GetWidth()),
+			static_cast<long>(InImage.GetHeight()),
+			Ordering,
+			BitDepth,
+			OCIO_NAMESPACE::AutoStride,
+			OCIO_NAMESPACE::AutoStride,
+			OCIO_NAMESPACE::AutoStride
+		);
+	};
 #endif
 }
 
@@ -286,115 +339,107 @@ void UOpenColorIOColorTransform::CacheResourceTextures()
 	if (Textures.IsEmpty())
 	{
 #if WITH_EDITOR && WITH_OCIO
-		OCIO_NAMESPACE::ConstConfigRcPtr CurrentConfig = ConfigurationOwner->GetNativeConfig_Internal()->Get();
-		if (CurrentConfig)
-		{
 #if !PLATFORM_EXCEPTIONS_DISABLED
-			try
+		try
 #endif
+		{
+			OCIO_NAMESPACE::ConstProcessorRcPtr TransformProcessor = GetTransformProcessor(this);
+			if (TransformProcessor)
 			{
-				OCIO_NAMESPACE::ConstProcessorRcPtr TransformProcessor = GetTransformProcessor(this, CurrentConfig, ContextKeyValues);
-				if (TransformProcessor)
+				OCIO_NAMESPACE::GpuShaderDescRcPtr ShaderDescription = OCIO_NAMESPACE::GpuShaderDesc::CreateShaderDesc();
+				ShaderDescription->setLanguage(OCIO_NAMESPACE::GPU_LANGUAGE_HLSL_DX11);
+				ShaderDescription->setFunctionName(StringCast<ANSICHAR>(OpenColorIOShader::OpenColorIOShaderFunctionName).Get());
+				ShaderDescription->setResourcePrefix("Ocio");
+
+				OCIO_NAMESPACE::ConstGPUProcessorRcPtr GPUProcessor = nullptr;
+				const UOpenColorIOSettings* Settings = GetDefault<UOpenColorIOSettings>();
+				check(Settings);
+
+				if (Settings->bUseLegacyProcessor)
 				{
-					OCIO_NAMESPACE::GpuShaderDescRcPtr ShaderDescription = OCIO_NAMESPACE::GpuShaderDesc::CreateShaderDesc();
-					ShaderDescription->setLanguage(OCIO_NAMESPACE::GPU_LANGUAGE_HLSL_DX11);
-					ShaderDescription->setFunctionName(StringCast<ANSICHAR>(OpenColorIOShader::OpenColorIOShaderFunctionName).Get());
-					ShaderDescription->setResourcePrefix("Ocio");
+					unsigned int EdgeLength = static_cast<unsigned int>(OpenColorIOShader::Lut3dEdgeLength);
+					GPUProcessor = TransformProcessor->getOptimizedLegacyGPUProcessor(GetProcessorOptimization(), EdgeLength);
+				}
+				else
+				{
+					GPUProcessor = TransformProcessor->getOptimizedGPUProcessor(GetProcessorOptimization());
+				}
+
+				GPUProcessor->extractGpuShaderInfo(ShaderDescription);
+				
+				ensureMsgf(ShaderDescription->getNumDynamicProperties() == 0, TEXT("Dynamic properties are not currently supported."));
 
-					OCIO_NAMESPACE::ConstGPUProcessorRcPtr GPUProcessor = nullptr;
-					const UOpenColorIOSettings* Settings = GetDefault<UOpenColorIOSettings>();
-					check(Settings);
+				const FString ProcessorID = StringCast<TCHAR>(GPUProcessor->getCacheID()).Get();
 
-					if (Settings->bUseLegacyProcessor)
-					{
-						unsigned int EdgeLength = static_cast<unsigned int>(OpenColorIOShader::Lut3dEdgeLength);
-						GPUProcessor = TransformProcessor->getOptimizedLegacyGPUProcessor(GetProcessorOptimization(), EdgeLength);
-					}
-					else
-					{
-						GPUProcessor = TransformProcessor->getOptimizedGPUProcessor(GetProcessorOptimization());
-					}
+				//In editor, it will use what's on DDC if there's something corresponding to the actual data or use that raw data
+				//that OCIO library has on board. The textures will be serialized only when cooking.
 
-					GPUProcessor->extractGpuShaderInfo(ShaderDescription);
-					// @todo: Remove once we add support for dynamic properties
-					ensureMsgf(ShaderDescription->getNumDynamicProperties() == 0, TEXT("Dynamic properties are not currently supported."));
+				// Process 3D luts
+				for (uint32 Index = 0; Index < ShaderDescription->getNum3DTextures(); ++Index)
+				{
+					const char* TextureName = nullptr;
+					const char* SamplerName = nullptr;
+					unsigned int EdgeLength;
+					OCIO_NAMESPACE::Interpolation Interpolation = OCIO_NAMESPACE::INTERP_TETRAHEDRAL;
 
-					const FString ProcessorID = StringCast<TCHAR>(GPUProcessor->getCacheID()).Get();
+					ShaderDescription->get3DTexture(Index, TextureName, SamplerName, EdgeLength, Interpolation);
+					checkf(TextureName && *TextureName && SamplerName && *SamplerName && EdgeLength > 0, TEXT("Invalid OCIO 3D texture or sampler."));
 
-					//In editor, it will use what's on DDC if there's something corresponding to the actual data or use that raw data
-					//that OCIO library has on board. The textures will be serialized only when cooking.
-					
-					// Process 3D luts
-					for (uint32 Index = 0; Index < ShaderDescription->getNum3DTextures(); ++Index)
-					{
-						const char* TextureName = nullptr;
-						const char* SamplerName = nullptr;
-						unsigned int EdgeLength;
-						OCIO_NAMESPACE::Interpolation Interpolation = OCIO_NAMESPACE::INTERP_TETRAHEDRAL;
-
-						ShaderDescription->get3DTexture(Index, TextureName, SamplerName, EdgeLength, Interpolation);
-						checkf(TextureName && *TextureName && SamplerName && *SamplerName && EdgeLength > 0, TEXT("Invalid OCIO 3D texture or sampler."));
-
-						const float* TextureValues = 0x0;
-						ShaderDescription->get3DTextureValues(Index, TextureValues);
-						checkf(TextureValues, TEXT("Failed to read OCIO 3D LUT data."));
-
-						TextureFilter Filter = TF_Bilinear;
-						if (Interpolation == OCIO_NAMESPACE::Interpolation::INTERP_NEAREST || Interpolation == OCIO_NAMESPACE::Interpolation::INTERP_TETRAHEDRAL)
-						{
-							Filter = TF_Nearest;
-						}
-						const FName TextureFName = FName(TextureName);
-						TObjectPtr<UTexture> Result = CreateTexture3DLUT(ProcessorID, TextureFName, EdgeLength, Filter, TextureValues);
-
-						const int32 SlotIndex = TextureFName.GetNumber() - 1; // Rely on FName's index number extraction for convenience
-						Textures.Add(SlotIndex, MoveTemp(Result));
-					}
+					const float* TextureValues = 0x0;
+					ShaderDescription->get3DTextureValues(Index, TextureValues);
+					checkf(TextureValues, TEXT("Failed to read OCIO 3D LUT data."));
 
-					// Process 1D luts
-					for (uint32 Index = 0; Index < ShaderDescription->getNumTextures(); ++Index)
+					TextureFilter Filter = TF_Bilinear;
+					if (Interpolation == OCIO_NAMESPACE::Interpolation::INTERP_NEAREST || Interpolation == OCIO_NAMESPACE::Interpolation::INTERP_TETRAHEDRAL)
 					{
-						const char* TextureName = nullptr;
-						const char* SamplerName = nullptr;
-						unsigned TextureWidth = 0;
-						unsigned TextureHeight = 0;
-						OCIO_NAMESPACE::GpuShaderDesc::TextureType Channel = OCIO_NAMESPACE::GpuShaderDesc::TEXTURE_RGB_CHANNEL;
-						OCIO_NAMESPACE::Interpolation Interpolation = OCIO_NAMESPACE::Interpolation::INTERP_LINEAR;
-
-						ShaderDescription->getTexture(Index, TextureName, SamplerName, TextureWidth, TextureHeight, Channel, Interpolation);
-						checkf(TextureName && *TextureName && SamplerName && *SamplerName && TextureWidth > 0, TEXT("Invalid OCIO 1D texture or sampler."));
-
-						const float* TextureValues = 0x0;
-						ShaderDescription->getTextureValues(Index, TextureValues);
-						checkf(TextureValues, TEXT("Failed to read OCIO 1D LUT data."));
-
-						TextureFilter Filter = Interpolation == OCIO_NAMESPACE::Interpolation::INTERP_NEAREST ? TF_Nearest : TF_Bilinear;
-						bool bRedChannelOnly = Channel == OCIO_NAMESPACE::GpuShaderCreator::TEXTURE_RED_CHANNEL;
-						const FName TextureFName = FName(TextureName);
-						TObjectPtr<UTexture> Result = CreateTexture1DLUT(ProcessorID, TextureFName, TextureWidth, TextureHeight, Filter, bRedChannelOnly, TextureValues);
-
-						const int32 SlotIndex = TextureFName.GetNumber() - 1; // Rely on FName's index number extraction for convenience
-						Textures.Add(SlotIndex, MoveTemp(Result));
+						Filter = TF_Nearest;
 					}
+					const FName TextureFName = FName(TextureName);
+					TObjectPtr<UTexture> Result = CreateTexture3DLUT(ProcessorID, TextureFName, EdgeLength, Filter, TextureValues);
 
-					ensureAlwaysMsgf(Textures.Num() <= (int32)OpenColorIOShader::MaximumTextureSlots, TEXT("Color transform %s exceeds our current limit of %u texture slots. Use the legacy processor instead."), *GetTransformFriendlyName(), OpenColorIOShader::MaximumTextureSlots);
+					const int32 SlotIndex = TextureFName.GetNumber() - 1; // Rely on FName's index number extraction for convenience
+					Textures.Add(SlotIndex, MoveTemp(Result));
 				}
-				else
+
+				// Process 1D luts
+				for (uint32 Index = 0; Index < ShaderDescription->getNumTextures(); ++Index)
 				{
-					UE_LOG(LogOpenColorIO, Error, TEXT("Failed to cache texture resource(s) for color transform %s. Transform processor was unusable."), *GetTransformFriendlyName());
+					const char* TextureName = nullptr;
+					const char* SamplerName = nullptr;
+					unsigned TextureWidth = 0;
+					unsigned TextureHeight = 0;
+					OCIO_NAMESPACE::GpuShaderDesc::TextureType Channel = OCIO_NAMESPACE::GpuShaderDesc::TEXTURE_RGB_CHANNEL;
+					OCIO_NAMESPACE::Interpolation Interpolation = OCIO_NAMESPACE::Interpolation::INTERP_LINEAR;
+
+					ShaderDescription->getTexture(Index, TextureName, SamplerName, TextureWidth, TextureHeight, Channel, Interpolation);
+					checkf(TextureName && *TextureName && SamplerName && *SamplerName && TextureWidth > 0, TEXT("Invalid OCIO 1D texture or sampler."));
+
+					const float* TextureValues = 0x0;
+					ShaderDescription->getTextureValues(Index, TextureValues);
+					checkf(TextureValues, TEXT("Failed to read OCIO 1D LUT data."));
+
+					TextureFilter Filter = Interpolation == OCIO_NAMESPACE::Interpolation::INTERP_NEAREST ? TF_Nearest : TF_Bilinear;
+					bool bRedChannelOnly = Channel == OCIO_NAMESPACE::GpuShaderCreator::TEXTURE_RED_CHANNEL;
+					const FName TextureFName = FName(TextureName);
+					TObjectPtr<UTexture> Result = CreateTexture1DLUT(ProcessorID, TextureFName, TextureWidth, TextureHeight, Filter, bRedChannelOnly, TextureValues);
+
+					const int32 SlotIndex = TextureFName.GetNumber() - 1; // Rely on FName's index number extraction for convenience
+					Textures.Add(SlotIndex, MoveTemp(Result));
 				}
+
+				ensureAlwaysMsgf(Textures.Num() <= (int32)OpenColorIOShader::MaximumTextureSlots, TEXT("Color transform %s exceeds our current limit of %u texture slots. Use the legacy processor instead."), *GetTransformFriendlyName(), OpenColorIOShader::MaximumTextureSlots);
 			}
-#if !PLATFORM_EXCEPTIONS_DISABLED
-			catch (OCIO_NAMESPACE::Exception& exception)
+			else
 			{
-				UE_LOG(LogOpenColorIO, Error, TEXT("Failed to cache texture resource(s) for color transform %s. Error message: %s"), *GetTransformFriendlyName(), StringCast<TCHAR>(exception.what()).Get());
+				UE_LOG(LogOpenColorIO, Error, TEXT("Failed to cache texture resource(s) for color transform %s. Configuration file [%s] was invalid."), *GetTransformFriendlyName(), *ConfigurationOwner->ConfigurationFile.FilePath);
 			}
-#endif
 		}
-		else
+#if !PLATFORM_EXCEPTIONS_DISABLED
+		catch (OCIO_NAMESPACE::Exception& exception)
 		{
-			UE_LOG(LogOpenColorIO, Error, TEXT("Failed to cache texture resource(s) for color transform %s. Configuration file [%s] was invalid."), *GetTransformFriendlyName(), *ConfigurationOwner->ConfigurationFile.FilePath);
+			UE_LOG(LogOpenColorIO, Error, TEXT("Failed to cache texture resource(s) for color transform %s. Error message: %s"), *GetTransformFriendlyName(), StringCast<TCHAR>(exception.what()).Get());
 		}
+#endif
 #endif
 	}
 }
@@ -564,6 +609,83 @@ bool UOpenColorIOColorTransform::IsTransform(const FString& InSourceColorSpace,
 	return false;
 }
 
+#if WITH_EDITOR
+bool UOpenColorIOColorTransform::EditorTransformImage(const FImageView& InOutImage) const
+{
+	// TODO: Handle extra conversions needed if working color space is the source or destination
+
+#if WITH_OCIO
+#if !PLATFORM_EXCEPTIONS_DISABLED
+	try
+#endif
+	{
+		OCIO_NAMESPACE::ConstProcessorRcPtr TransformProcessor = GetTransformProcessor(this);
+		if (TransformProcessor)
+		{
+			TUniquePtr<OCIO_NAMESPACE::PackedImageDesc> ImageDesc = GetImageDesc(InOutImage);
+			if (ImageDesc)
+			{
+				OCIO_NAMESPACE::ConstCPUProcessorRcPtr Processor = TransformProcessor->getOptimizedCPUProcessor(
+					ImageDesc->getBitDepth(),
+					ImageDesc->getBitDepth(),
+					OCIO_NAMESPACE::OptimizationFlags::OPTIMIZATION_DEFAULT
+				);
+				Processor->apply(*ImageDesc);
+				
+				return true;
+			}
+		}
+	}
+#if !PLATFORM_EXCEPTIONS_DISABLED
+	catch (OCIO_NAMESPACE::Exception& exception)
+	{
+		UE_LOG(LogOpenColorIO, Log, TEXT("Failed to transform image. Error message: %s"), StringCast<TCHAR>(exception.what()).Get());
+	}
+#endif
+#endif // WITH_OCIO
+
+	return false;
+}
+
+bool UOpenColorIOColorTransform::EditorTransformImage(const FImageView& SrcImage, const FImageView& DestImage) const
+{
+	// TODO: Handle extra conversions needed if working color space is the source or destination
+
+#if WITH_OCIO
+#if !PLATFORM_EXCEPTIONS_DISABLED
+	try
+#endif
+	{
+		OCIO_NAMESPACE::ConstProcessorRcPtr TransformProcessor = GetTransformProcessor(this);
+		if (TransformProcessor)
+		{
+			TUniquePtr<OCIO_NAMESPACE::PackedImageDesc> SrcImageDesc = GetImageDesc(SrcImage);
+			TUniquePtr<OCIO_NAMESPACE::PackedImageDesc> DestImageDesc = GetImageDesc(DestImage);
+			if (SrcImageDesc && DestImageDesc)
+			{
+				OCIO_NAMESPACE::ConstCPUProcessorRcPtr Processor = TransformProcessor->getOptimizedCPUProcessor(
+					SrcImageDesc->getBitDepth(),
+					DestImageDesc->getBitDepth(),
+					OCIO_NAMESPACE::OptimizationFlags::OPTIMIZATION_DEFAULT
+				);
+				TransformProcessor->getDefaultCPUProcessor()->apply(*SrcImageDesc, *DestImageDesc);
+
+				return true;
+			}
+		}
+	}
+#if !PLATFORM_EXCEPTIONS_DISABLED
+	catch (OCIO_NAMESPACE::Exception& exception)
+	{
+		UE_LOG(LogOpenColorIO, Log, TEXT("Failed to transform image. Error message: %s"), StringCast<TCHAR>(exception.what()).Get());
+	}
+#endif
+#endif // WITH_OCIO
+
+	return false;
+}
+#endif // WITH_EDITOR
+
 bool UOpenColorIOColorTransform::GetDisplayViewDirection(EOpenColorIOViewTransformDirection& OutDirection) const
 {
 	if (bIsDisplayViewType)
@@ -652,61 +774,54 @@ bool UOpenColorIOColorTransform::UpdateShaderInfo(FString& OutShaderCodeHash, FS
 {
 #if WITH_EDITOR
 #if WITH_OCIO
-	OCIO_NAMESPACE::ConstConfigRcPtr CurrentConfig = ConfigurationOwner->GetNativeConfig_Internal()->Get();
-	if (CurrentConfig)
-	{
 #if !PLATFORM_EXCEPTIONS_DISABLED
-		try
+	try
 #endif
+	{
+		OCIO_NAMESPACE::ConstProcessorRcPtr TransformProcessor = GetTransformProcessor(this);
+		if (TransformProcessor)
 		{
-			OCIO_NAMESPACE::ConstProcessorRcPtr TransformProcessor = GetTransformProcessor(this, CurrentConfig, ContextKeyValues);
-			if (TransformProcessor)
-			{
-				OCIO_NAMESPACE::GpuShaderDescRcPtr ShaderDescription = OCIO_NAMESPACE::GpuShaderDesc::CreateShaderDesc();
-				ShaderDescription->setLanguage(OCIO_NAMESPACE::GPU_LANGUAGE_HLSL_DX11);
-				ShaderDescription->setFunctionName(StringCast<ANSICHAR>(OpenColorIOShader::OpenColorIOShaderFunctionName).Get());
-				ShaderDescription->setResourcePrefix("Ocio");
+			OCIO_NAMESPACE::GpuShaderDescRcPtr ShaderDescription = OCIO_NAMESPACE::GpuShaderDesc::CreateShaderDesc();
+			ShaderDescription->setLanguage(OCIO_NAMESPACE::GPU_LANGUAGE_HLSL_DX11);
+			ShaderDescription->setFunctionName(StringCast<ANSICHAR>(OpenColorIOShader::OpenColorIOShaderFunctionName).Get());
+			ShaderDescription->setResourcePrefix("Ocio");
 
-				OCIO_NAMESPACE::ConstGPUProcessorRcPtr GPUProcessor = nullptr;
-				const UOpenColorIOSettings* Settings = GetDefault<UOpenColorIOSettings>();
-				check(Settings);
+			OCIO_NAMESPACE::ConstGPUProcessorRcPtr GPUProcessor = nullptr;
+			const UOpenColorIOSettings* Settings = GetDefault<UOpenColorIOSettings>();
+			check(Settings);
 
-				if (Settings->bUseLegacyProcessor)
-				{
-					unsigned int EdgeLength = static_cast<unsigned int>(OpenColorIOShader::Lut3dEdgeLength);
-					GPUProcessor = TransformProcessor->getOptimizedLegacyGPUProcessor(GetProcessorOptimization(), EdgeLength);
-				}
-				else
-				{
-					GPUProcessor = TransformProcessor->getOptimizedGPUProcessor(GetProcessorOptimization());
-				}
-				GPUProcessor->extractGpuShaderInfo(ShaderDescription);
-				// @todo: Remove once we add support for dynamic properties
-				ensureMsgf(ShaderDescription->getNumDynamicProperties() == 0, TEXT("We do not currently support dynamic properties."));
-
-				FString GLSLShaderCode = StringCast<TCHAR>(ShaderDescription->getShaderText()).Get();
-
-				OutShaderCodeHash = StringCast<TCHAR>(ShaderDescription->getCacheID()).Get();
-				OutShaderCode = StringCast<TCHAR>(ShaderDescription->getShaderText()).Get();
-				OutRawConfigHash = StringCast<TCHAR>(CurrentConfig->getCacheID()).Get();
-				return true;
+			if (Settings->bUseLegacyProcessor)
+			{
+				unsigned int EdgeLength = static_cast<unsigned int>(OpenColorIOShader::Lut3dEdgeLength);
+				GPUProcessor = TransformProcessor->getOptimizedLegacyGPUProcessor(GetProcessorOptimization(), EdgeLength);
 			}
 			else
 			{
-				UE_LOG(LogOpenColorIO, Error, TEXT("Failed to fetch shader info for color transform %s. Transform processor was unusable."), *GetTransformFriendlyName());
+				GPUProcessor = TransformProcessor->getOptimizedGPUProcessor(GetProcessorOptimization());
 			}
+			GPUProcessor->extractGpuShaderInfo(ShaderDescription);
+
+			ensureMsgf(ShaderDescription->getNumDynamicProperties() == 0, TEXT("We do not currently support dynamic properties."));
+
+			FString GLSLShaderCode = StringCast<TCHAR>(ShaderDescription->getShaderText()).Get();
+
+			OutShaderCodeHash = StringCast<TCHAR>(ShaderDescription->getCacheID()).Get();
+			OutShaderCode = StringCast<TCHAR>(ShaderDescription->getShaderText()).Get();
+			OCIO_NAMESPACE::ConstConfigRcPtr CurrentConfig = ConfigurationOwner->GetNativeConfig_Internal()->Get();
+			OutRawConfigHash = StringCast<TCHAR>(CurrentConfig->getCacheID()).Get();
+			return true;
 		}
-#if !PLATFORM_EXCEPTIONS_DISABLED
-		catch (OCIO_NAMESPACE::Exception& exception)
+		else
 		{
-			UE_LOG(LogOpenColorIO, Error, TEXT("Failed to fetch shader info for color transform %s. Error message: %s"), *GetTransformFriendlyName(), StringCast<TCHAR>(exception.what()).Get());
+			UE_LOG(LogOpenColorIO, Error, TEXT("Failed to fetch shader info for color transform %s. Configuration file [%s] was invalid."), *GetTransformFriendlyName(), *ConfigurationOwner->ConfigurationFile.FilePath);
 		}
-#endif
 	}
-	else
+#if !PLATFORM_EXCEPTIONS_DISABLED
+	catch (OCIO_NAMESPACE::Exception& exception)
 	{
-		UE_LOG(LogOpenColorIO, Error, TEXT("Failed to fetch shader info for color transform %s. Configuration file [%s] was invalid."), *GetTransformFriendlyName(), *ConfigurationOwner->ConfigurationFile.FilePath);
+		UE_LOG(LogOpenColorIO, Error, TEXT("Failed to fetch shader info for color transform %s. Error message: %s"), *GetTransformFriendlyName(), StringCast<TCHAR>(exception.what()).Get());
 	}
+#endif
 
 	return false;
 #else
diff --git a/Engine/Plugins/Compositing/OpenColorIO/Source/OpenColorIO/Private/OpenColorIOConfiguration.cpp b/Engine/Plugins/Compositing/OpenColorIO/Source/OpenColorIO/Private/OpenColorIOConfiguration.cpp
--- a/Engine/Plugins/Compositing/OpenColorIO/Source/OpenColorIO/Private/OpenColorIOConfiguration.cpp
+++ b/Engine/Plugins/Compositing/OpenColorIO/Source/OpenColorIO/Private/OpenColorIOConfiguration.cpp
@@ -160,6 +160,28 @@ bool UOpenColorIOConfiguration::Validate() const
 #endif // WITH_EDITOR
 }
 
+#if WITH_EDITOR
+bool UOpenColorIOConfiguration::EditorTransformImage(const FOpenColorIOColorConversionSettings& InSettings, const FImageView& InOutImage) const
+{
+	if (const TObjectPtr<UOpenColorIOColorTransform>* TransformPtr = FindTransform(InSettings))
+	{
+		return (*TransformPtr)->EditorTransformImage(InOutImage);
+	}
+
+	return false;
+}
+
+bool UOpenColorIOConfiguration::EditorTransformImage(const FOpenColorIOColorConversionSettings& InSettings, const FImageView& SrcImage, const FImageView& DestImage) const
+{
+	if (const TObjectPtr<UOpenColorIOColorTransform>* TransformPtr = FindTransform(InSettings))
+	{
+		return (*TransformPtr)->EditorTransformImage(SrcImage, DestImage);
+	}
+
+	return false;
+}
+#endif
+
 void UOpenColorIOConfiguration::ReloadExistingColorspaces()
 {
 #if WITH_EDITOR && WITH_OCIO
diff --git a/Engine/Plugins/Compositing/OpenColorIO/Source/OpenColorIO/Public/OpenColorIOColorTransform.h b/Engine/Plugins/Compositing/OpenColorIO/Source/OpenColorIO/Public/OpenColorIOColorTransform.h
--- a/Engine/Plugins/Compositing/OpenColorIO/Source/OpenColorIO/Public/OpenColorIOColorTransform.h
+++ b/Engine/Plugins/Compositing/OpenColorIO/Source/OpenColorIO/Public/OpenColorIOColorTransform.h
@@ -79,6 +79,14 @@ public:
 	bool IsTransform(const FString& InSourceColorSpace, const FString& InDestinationColorSpace) const;
 	bool IsTransform(const FString& InSourceColorSpace, const FString& InDisplay, const FString& InView, EOpenColorIOViewTransformDirection InDirection) const;
 
+#if WITH_EDITOR
+	/** Apply the color transform in-place to the specified image. */
+	bool EditorTransformImage(const FImageView& InOutImage) const;
+	
+	/** Apply the color transform from the source image to the destination image. (The destination FImageView is const but what it points at is not.) */
+	bool EditorTransformImage(const FImageView& SrcImage, const FImageView& DestImage) const;
+#endif
+
 	/**
 	 * Get the display view direction type, when applicable.
 	 *
@@ -90,7 +98,9 @@ public:
 	// For all ColorTransforms, UOpenColorIOColorTransform::CacheResourceShadersForRendering
 	static void AllColorTransformsCacheResourceShadersForRendering();
 
-	
+	// Get the context key-values.
+	const TMap<FString, FString>& GetContextKeyValues() const { return ContextKeyValues; }
+
 
 protected:
 
diff --git a/Engine/Plugins/Compositing/OpenColorIO/Source/OpenColorIO/Public/OpenColorIOConfiguration.h b/Engine/Plugins/Compositing/OpenColorIO/Source/OpenColorIO/Public/OpenColorIOConfiguration.h
--- a/Engine/Plugins/Compositing/OpenColorIO/Source/OpenColorIO/Public/OpenColorIOConfiguration.h
+++ b/Engine/Plugins/Compositing/OpenColorIO/Source/OpenColorIO/Public/OpenColorIOConfiguration.h
@@ -48,6 +48,14 @@ public:
 	bool HasDesiredDisplayView(const FOpenColorIODisplayView& DisplayView) const;
 	bool Validate() const;
 
+#if WITH_EDITOR
+	/** Apply the color transform in-place to the specified image. */
+	bool EditorTransformImage(const FOpenColorIOColorConversionSettings& InSettings, const FImageView& InOutImage) const;
+
+	/** Apply the color transform from the source image to the destination image. (The destination FImageView is const but what it points at is not.) */
+	bool EditorTransformImage(const FOpenColorIOColorConversionSettings& InSettings, const FImageView& SrcImage, const FImageView& DestImage) const;
+#endif
+
 	/** This forces to reload colorspaces and corresponding shaders if those are not loaded already. */
 	UFUNCTION(BlueprintCallable, Category = "OpenColorIO")
 	void ReloadExistingColorspaces();
