OCIO: Fix NoPCH issue with FImageView, related to CL 24368324.

#rnx

diff --git a/Engine/Plugins/Compositing/OpenColorIO/Source/OpenColorIO/Public/OpenColorIOColorTransform.h b/Engine/Plugins/Compositing/OpenColorIO/Source/OpenColorIO/Public/OpenColorIOColorTransform.h
--- a/Engine/Plugins/Compositing/OpenColorIO/Source/OpenColorIO/Public/OpenColorIOColorTransform.h
+++ b/Engine/Plugins/Compositing/OpenColorIO/Source/OpenColorIO/Public/OpenColorIOColorTransform.h
@@ -15,6 +15,7 @@
 #include "OpenColorIOColorTransform.generated.h"
 
 
+struct FImageView;
 class UOpenColorIOConfiguration;
 
 
diff --git a/Engine/Plugins/Compositing/OpenColorIO/Source/OpenColorIO/Public/OpenColorIOConfiguration.h b/Engine/Plugins/Compositing/OpenColorIO/Source/OpenColorIO/Public/OpenColorIOConfiguration.h
--- a/Engine/Plugins/Compositing/OpenColorIO/Source/OpenColorIO/Public/OpenColorIOConfiguration.h
+++ b/Engine/Plugins/Compositing/OpenColorIO/Source/OpenColorIO/Public/OpenColorIOConfiguration.h
@@ -18,6 +18,7 @@ class FOpenColorIONativeConfiguration;
 class FOpenColorIOTransformResource;
 class FTextureResource;
 class UOpenColorIOColorTransform;
+struct FImageView;
 struct FFileChangeData;
 class SNotificationItem;
 namespace ERHIFeatureLevel { enum Type : int; }
