OCIO: Properly moving native configuration back to the private folder by exposing the necessary data for the editor menus.

The file was made public in 24223516 due to engine cleanup.

#rnx

diff --git a/Engine/Plugins/Compositing/OpenColorIO/Source/OpenColorIO/Private/OpenColorIOConfiguration.cpp b/Engine/Plugins/Compositing/OpenColorIO/Source/OpenColorIO/Private/OpenColorIOConfiguration.cpp
--- a/Engine/Plugins/Compositing/OpenColorIO/Source/OpenColorIO/Private/OpenColorIOConfiguration.cpp
+++ b/Engine/Plugins/Compositing/OpenColorIO/Source/OpenColorIO/Private/OpenColorIOConfiguration.cpp
@@ -685,5 +685,105 @@ void UOpenColorIOConfiguration::OnToastCallback(bool bInReloadColorspaces)
 	}
 }
 
-#undef LOCTEXT_NAMESPACE
+#if WITH_EDITOR
+FOpenColorIOEditorConfigurationInspector::FOpenColorIOEditorConfigurationInspector(const UOpenColorIOConfiguration& InConfiguration)
+	: NativeConfig(*InConfiguration.GetNativeConfig_Internal()) // Private implementation pointer is always present and can be safely dereferenced
+{
+}
+
+int32 FOpenColorIOEditorConfigurationInspector::GetNumColorSpaces() const
+{
+#if WITH_OCIO
+	if (OCIO_NAMESPACE::ConstConfigRcPtr Config = NativeConfig.Get())
+	{
+		return Config->getNumColorSpaces(OCIO_NAMESPACE::SEARCH_REFERENCE_SPACE_ALL, OCIO_NAMESPACE::COLORSPACE_ACTIVE);
+	}
+#endif
+
+	return 0;
+}
+
+FString FOpenColorIOEditorConfigurationInspector::GetColorSpaceName(int32 Index) const
+{
+#if WITH_OCIO
+	if (OCIO_NAMESPACE::ConstConfigRcPtr Config = NativeConfig.Get())
+	{
+		const char* ColorSpaceName = Config->getColorSpaceNameByIndex(OCIO_NAMESPACE::SEARCH_REFERENCE_SPACE_ALL, OCIO_NAMESPACE::COLORSPACE_ACTIVE, Index);
+
+		return StringCast<TCHAR>(ColorSpaceName).Get();
+	}
+#endif
 
+	return {};
+}
+
+FString FOpenColorIOEditorConfigurationInspector::GetColorSpaceFamilyName(const TCHAR* InColorSpaceName) const
+{
+#if WITH_OCIO
+	if (OCIO_NAMESPACE::ConstConfigRcPtr Config = NativeConfig.Get())
+	{
+		OCIO_NAMESPACE::ConstColorSpaceRcPtr ColorSpace = Config->getColorSpace(StringCast<ANSICHAR>(InColorSpaceName).Get());
+		if (ColorSpace != nullptr)
+		{
+			return StringCast<TCHAR>(ColorSpace->getFamily()).Get();
+		}
+	}
+#endif
+
+	return {};
+}
+
+int32 FOpenColorIOEditorConfigurationInspector::GetNumDisplays() const
+{
+#if WITH_OCIO
+	if (OCIO_NAMESPACE::ConstConfigRcPtr Config = NativeConfig.Get())
+	{
+		return Config->getNumDisplays();
+	}
+#endif
+
+	return 0;
+}
+
+FString FOpenColorIOEditorConfigurationInspector::GetDisplayName(int32 Index) const
+{
+#if WITH_OCIO
+	if (OCIO_NAMESPACE::ConstConfigRcPtr Config = NativeConfig.Get())
+	{
+		const char* DisplayName = Config->getDisplay(Index);
+
+		return StringCast<TCHAR>(DisplayName).Get();
+	}
+#endif
+	
+	return {};
+}
+
+int32 FOpenColorIOEditorConfigurationInspector::GetNumViews(const TCHAR* InDisplayName) const
+{
+#if WITH_OCIO
+	if (OCIO_NAMESPACE::ConstConfigRcPtr Config = NativeConfig.Get())
+	{
+		return Config->getNumViews(StringCast<ANSICHAR>(InDisplayName).Get());
+	}
+#endif
+
+	return 0;
+}
+
+FString FOpenColorIOEditorConfigurationInspector::GetViewName(const TCHAR* InDisplayName, int32 Index) const
+{
+#if WITH_OCIO
+	if (OCIO_NAMESPACE::ConstConfigRcPtr Config = NativeConfig.Get())
+	{
+		const char* ViewName = Config->getView(StringCast<ANSICHAR>(InDisplayName).Get(), Index);
+
+		return StringCast<TCHAR>(ViewName).Get();
+	}
+#endif
+	
+	return {};
+}
+#endif //WITH_EDITOR
+
+#undef LOCTEXT_NAMESPACE
diff --git a/Engine/Plugins/Compositing/OpenColorIO/Source/OpenColorIO/Private/OpenColorIONativeConfiguration.h b/Engine/Plugins/Compositing/OpenColorIO/Source/OpenColorIO/Private/OpenColorIONativeConfiguration.h
new file mode 100644
--- /dev/null
+++ b/Engine/Plugins/Compositing/OpenColorIO/Source/OpenColorIO/Private/OpenColorIONativeConfiguration.h
@@ -0,0 +1,55 @@
+// Copyright Epic Games, Inc. All Rights Reserved.
+
+#pragma once
+
+#if WITH_EDITOR && WITH_OCIO
+#include "OpenColorIO/OpenColorIO.h"
+#endif
+
+class FOpenColorIONativeConfiguration {
+public:
+
+#if WITH_EDITOR && WITH_OCIO
+	OCIO_NAMESPACE::ConstConfigRcPtr Get() const
+	{
+		return Config;
+	}
+
+	void Set(OCIO_NAMESPACE::ConstConfigRcPtr InConfig, FStringView InDefaultWorkingColorSpaceName = FStringView())
+	{
+		using namespace OCIO_NAMESPACE;
+
+		if (InConfig == nullptr)
+		{
+			Config.reset();
+		}
+		else
+		{	
+			ConstColorSpaceRcPtr InterchangeCS = InConfig->getColorSpace(InConfig->getCanonicalName("aces_interchange"/*OCIO_NAMESPACE::ROLE_INTERCHANGE_SCENE*/));
+
+			// When the aces interchange color space is present, we add the working color space as an additional option.
+			if (InterchangeCS != nullptr && InConfig->getColorSpace(StringCast<ANSICHAR>(InDefaultWorkingColorSpaceName.GetData()).Get()) == nullptr)
+			{
+				ColorSpaceRcPtr WorkingCS = InterchangeCS->createEditableCopy();
+				WorkingCS->setName(StringCast<ANSICHAR>(InDefaultWorkingColorSpaceName.GetData()).Get());
+				WorkingCS->setFamily("UE");
+				WorkingCS->clearAliases();
+
+				ConfigRcPtr ConfigCopy = InConfig->createEditableCopy();
+				ConfigCopy->addColorSpace(WorkingCS);
+
+				Config = ConfigCopy;
+			}
+			else
+			{
+				Config = InConfig;
+			}
+		}
+	}
+#endif
+
+private:
+#if WITH_EDITORONLY_DATA && WITH_OCIO
+	OCIO_NAMESPACE::ConstConfigRcPtr Config;
+#endif //WITH_EDITORONLY_DATA
+};
diff --git a/Engine/Plugins/Compositing/OpenColorIO/Source/OpenColorIO/Public/OpenColorIOConfiguration.h b/Engine/Plugins/Compositing/OpenColorIO/Source/OpenColorIO/Public/OpenColorIOConfiguration.h
--- a/Engine/Plugins/Compositing/OpenColorIO/Source/OpenColorIO/Public/OpenColorIOConfiguration.h
+++ b/Engine/Plugins/Compositing/OpenColorIO/Source/OpenColorIO/Public/OpenColorIOConfiguration.h
@@ -141,3 +141,36 @@ private:
 	/** Private implementation of the native OpenColorIO config object. */
 	TPimplPtr<FOpenColorIONativeConfiguration> NativeConfig;
 };
+
+
+#if WITH_EDITOR
+/**
+ * Editor-only class to access native configuration names of color spaces, displays and views for menus.
+ */
+class FOpenColorIOEditorConfigurationInspector {
+public:
+
+	OPENCOLORIO_API FOpenColorIOEditorConfigurationInspector(const UOpenColorIOConfiguration& InConfiguration);
+
+	/** Get the number of color spaces in the configuration. */
+	OPENCOLORIO_API int32 GetNumColorSpaces() const;
+	/** Get a color space name at an index. */
+	OPENCOLORIO_API FString GetColorSpaceName(int32 Index) const;
+	/** Get the family name for a color space. */
+	OPENCOLORIO_API FString GetColorSpaceFamilyName(const TCHAR* InColorSpaceName) const;
+
+	/** Get the number of displays in the configuration. */
+	OPENCOLORIO_API int32 GetNumDisplays() const;
+	/** Get a display name at an index. */
+	OPENCOLORIO_API FString GetDisplayName(int32 Index) const;
+
+	/** Get the number of views for a display. */
+	OPENCOLORIO_API int32 GetNumViews(const TCHAR* InDisplayName) const;
+	/** Get a view name for its display and index. */
+	OPENCOLORIO_API FString GetViewName(const TCHAR* InDisplayName, int32 Index) const;
+
+private:
+	/** Reference to the native OCIO config. */
+	const FOpenColorIONativeConfiguration& NativeConfig;
+};
+#endif //WITH_EDITOR
diff --git a/Engine/Plugins/Compositing/OpenColorIO/Source/OpenColorIO/Public/OpenColorIONativeConfiguration.h b/Engine/Plugins/Compositing/OpenColorIO/Source/OpenColorIO/Public/OpenColorIONativeConfiguration.h
deleted file mode 100644
diff --git a/Engine/Plugins/Compositing/OpenColorIO/Source/OpenColorIOEditor/Private/OpenColorIOColorSpaceCustomization.cpp b/Engine/Plugins/Compositing/OpenColorIO/Source/OpenColorIOEditor/Private/OpenColorIOColorSpaceCustomization.cpp
--- a/Engine/Plugins/Compositing/OpenColorIO/Source/OpenColorIOEditor/Private/OpenColorIOColorSpaceCustomization.cpp
+++ b/Engine/Plugins/Compositing/OpenColorIO/Source/OpenColorIOEditor/Private/OpenColorIOColorSpaceCustomization.cpp
@@ -7,7 +7,6 @@
 #include "IPropertyUtilities.h"
 #include "PropertyHandle.h"
 #include "OpenColorIOConfiguration.h"
-#include "OpenColorIONativeConfiguration.h"
 #include "OpenColorIO/OpenColorIO.h"
 #include "Widgets/Input/SComboButton.h"
 
@@ -16,7 +15,7 @@
 void FOpenColorIOColorSpaceCustomization::CustomizeHeader(TSharedRef<IPropertyHandle> InPropertyHandle, FDetailWidgetRow& HeaderRow, IPropertyTypeCustomizationUtils& CustomizationUtils)
 {
 	//Reset internals
-	CachedNativeConfig = nullptr;
+	ConfigInspector.Reset(nullptr);
 
 	CachedProperty = InPropertyHandle;
 
@@ -66,7 +65,7 @@ void FOpenColorIOColorSpaceCustomization::CustomizeHeader(TSharedRef<IPropertyHa
 void FOpenColorIODisplayViewCustomization::CustomizeHeader(TSharedRef<IPropertyHandle> InPropertyHandle, FDetailWidgetRow& HeaderRow, IPropertyTypeCustomizationUtils& CustomizationUtils)
 {
 	//Reset internals
-	CachedNativeConfig = nullptr;
+	ConfigInspector.Reset(nullptr);
 
 	CachedProperty = InPropertyHandle;
 
@@ -118,7 +117,7 @@ IPropertyTypeCustomizationOpenColorIO::IPropertyTypeCustomizationOpenColorIO(TSh
 {
 }
 
-FOpenColorIONativeConfiguration* IPropertyTypeCustomizationOpenColorIO::GetNativeConfig(const TSharedPtr<IPropertyHandle>& InConfigurationAssetProperty)
+TUniquePtr<FOpenColorIOEditorConfigurationInspector> IPropertyTypeCustomizationOpenColorIO::GetConfigurationInspector(const TSharedPtr<IPropertyHandle>& InConfigurationAssetProperty)
 {
 	if (InConfigurationAssetProperty.IsValid())
 	{
@@ -129,7 +128,7 @@ FOpenColorIONativeConfiguration* IPropertyTypeCustomizationOpenColorIO::GetNativ
 		UOpenColorIOConfiguration* Configuration = reinterpret_cast<UOpenColorIOConfiguration*>(RawData[0]);
 		check(Configuration);
 		
-		return Configuration->GetNativeConfig_Internal();		
+		return MakeUnique<FOpenColorIOEditorConfigurationInspector>(*Configuration);
 	}
 
 	return nullptr;
@@ -170,20 +169,12 @@ void FOpenColorIOColorSpaceCustomization::PopulateSubMenu(FMenuBuilder& InMenuBu
 	// To keep track of submenus that were already added
 	TArray<FString> ExistingSubMenus;
 
-	const int32 ColorSpaceCount = CachedNativeConfig->Get()->getNumColorSpaces(OCIO_NAMESPACE::SEARCH_REFERENCE_SPACE_ALL, OCIO_NAMESPACE::COLORSPACE_ACTIVE);
-	for (int32 i = 0; i < ColorSpaceCount; ++i)
+	for (int32 i = 0; i < ConfigInspector->GetNumColorSpaces(); ++i)
 	{
-		const char* ColorSpaceName = CachedNativeConfig->Get()->getColorSpaceNameByIndex(OCIO_NAMESPACE::SEARCH_REFERENCE_SPACE_ALL, OCIO_NAMESPACE::COLORSPACE_ACTIVE, i);
-		OCIO_NAMESPACE::ConstColorSpaceRcPtr LibColorSpace = CachedNativeConfig->Get()->getColorSpace(ColorSpaceName);
-		if (!LibColorSpace)
-		{
-			continue;
-		}
-
 		FOpenColorIOColorSpace ColorSpace;
 		ColorSpace.ColorSpaceIndex = i;
-		ColorSpace.ColorSpaceName = StringCast<TCHAR>(ColorSpaceName).Get();
-		ColorSpace.FamilyName = StringCast<TCHAR>(LibColorSpace->getFamily()).Get();
+		ColorSpace.ColorSpaceName = ConfigInspector->GetColorSpaceName(i);
+		ColorSpace.FamilyName = ConfigInspector->GetColorSpaceFamilyName(*ColorSpace.ColorSpaceName);
 		
 		//Filter out color spaces that don't belong to this hierarchy
 		if (InPreviousFamilyHierarchy.IsEmpty() || ColorSpace.FamilyName.Contains(InPreviousFamilyHierarchy))
@@ -230,7 +221,7 @@ void FOpenColorIOColorSpaceCustomization::AddMenuEntry(FMenuBuilder& InMenuBuild
 
 TSharedRef<SWidget> FOpenColorIOColorSpaceCustomization::HandleSourceComboButtonMenuContent()
 {
-	CachedNativeConfig = GetNativeConfig(ConfigurationObjectProperty);
+	ConfigInspector = GetConfigurationInspector(ConfigurationObjectProperty);
 
 	// Generate menu
 	FMenuBuilder MenuBuilder(true, nullptr);
@@ -238,22 +229,16 @@ TSharedRef<SWidget> FOpenColorIOColorSpaceCustomization::HandleSourceComboButton
 
 	MenuBuilder.BeginSection("AllColorSpaces", LOCTEXT("AllColorSpacesSection", "ColorSpaces"));
 	{
-		if (CachedNativeConfig && CachedNativeConfig->Get() != nullptr)
+		if (ConfigInspector)
 		{
-			const int32 ColorSpaceCount = CachedNativeConfig->Get()->getNumColorSpaces(OCIO_NAMESPACE::SEARCH_REFERENCE_SPACE_ALL, OCIO_NAMESPACE::COLORSPACE_ACTIVE);
+			const int32 ColorSpaceCount = ConfigInspector->GetNumColorSpaces();
+
 			for (int32 i = 0; i < ColorSpaceCount; ++i)
 			{
-				const char* ColorSpaceName = CachedNativeConfig->Get()->getColorSpaceNameByIndex(OCIO_NAMESPACE::SEARCH_REFERENCE_SPACE_ALL, OCIO_NAMESPACE::COLORSPACE_ACTIVE, i);
-				OCIO_NAMESPACE::ConstColorSpaceRcPtr LibColorSpace = CachedNativeConfig->Get()->getColorSpace(ColorSpaceName);
-				if (!LibColorSpace)
-				{
-					continue;
-				}
-
 				FOpenColorIOColorSpace ColorSpace;
 				ColorSpace.ColorSpaceIndex = i;
-				ColorSpace.ColorSpaceName = StringCast<TCHAR>(ColorSpaceName).Get();
-				ColorSpace.FamilyName = StringCast<TCHAR>(LibColorSpace->getFamily()).Get();
+				ColorSpace.ColorSpaceName = ConfigInspector->GetColorSpaceName(i);
+				ColorSpace.FamilyName = ConfigInspector->GetColorSpaceFamilyName(*ColorSpace.ColorSpaceName);
 
 				//Top level menus have no preceding hierarchy.
 				const int32 CurrentMenuDepth = 0;
@@ -278,11 +263,9 @@ TSharedRef<SWidget> FOpenColorIOColorSpaceCustomization::HandleSourceComboButton
 
 void FOpenColorIODisplayViewCustomization::PopulateViewSubMenu(FMenuBuilder& InMenuBuilder, FOpenColorIODisplayView InDisplayView)
 {
-	const int32 ViewCount = CachedNativeConfig->Get()->getNumViews(TCHAR_TO_ANSI(*InDisplayView.Display));
-	for (int32 i = 0; i < ViewCount; ++i)
+	for (int32 i = 0; i < ConfigInspector->GetNumViews(*InDisplayView.Display); ++i)
 	{
-		const char* ViewName = CachedNativeConfig->Get()->getView(TCHAR_TO_ANSI(*InDisplayView.Display), i);
-		InDisplayView.View = StringCast<TCHAR>(ViewName).Get();
+		InDisplayView.View = ConfigInspector->GetViewName(*InDisplayView.Display, i);
 		
 		AddMenuEntry(InMenuBuilder, InDisplayView);
 	}
@@ -325,7 +308,7 @@ void FOpenColorIODisplayViewCustomization::AddMenuEntry(FMenuBuilder& InMenuBuil
 
 TSharedRef<SWidget> FOpenColorIODisplayViewCustomization::HandleSourceComboButtonMenuContent()
 {
-	CachedNativeConfig = GetNativeConfig(ConfigurationObjectProperty);
+	ConfigInspector = GetConfigurationInspector(ConfigurationObjectProperty);
 
 	// Generate menu
 	FMenuBuilder MenuBuilder(true, nullptr);
@@ -333,14 +316,13 @@ TSharedRef<SWidget> FOpenColorIODisplayViewCustomization::HandleSourceComboButto
 
 	MenuBuilder.BeginSection("AllDisplayViews", LOCTEXT("AllDisplayViewsSection", "Display - View"));
 	{
-		if (CachedNativeConfig && CachedNativeConfig->Get() != nullptr)
+		if (ConfigInspector)
 		{
-			const int32 DisplayCount = CachedNativeConfig->Get()->getNumDisplays();
+			const int32 DisplayCount = ConfigInspector->GetNumDisplays();
+
 			for (int32 i = 0; i < DisplayCount; ++i)
 			{
-				const char* DisplayName = CachedNativeConfig->Get()->getDisplay(i);
-
-				const FOpenColorIODisplayView DisplayViewValue = { StringCast<TCHAR>(DisplayName).Get(), TEXT("<Invalid>") };
+				const FOpenColorIODisplayView DisplayViewValue = { ConfigInspector->GetDisplayName(i), TEXT("<Invalid>")};
 
 				MenuBuilder.AddSubMenu(
 					FText::FromString(DisplayViewValue.Display),
diff --git a/Engine/Plugins/Compositing/OpenColorIO/Source/OpenColorIOEditor/Private/OpenColorIOColorSpaceCustomization.h b/Engine/Plugins/Compositing/OpenColorIO/Source/OpenColorIOEditor/Private/OpenColorIOColorSpaceCustomization.h
--- a/Engine/Plugins/Compositing/OpenColorIO/Source/OpenColorIOEditor/Private/OpenColorIOColorSpaceCustomization.h
+++ b/Engine/Plugins/Compositing/OpenColorIO/Source/OpenColorIOEditor/Private/OpenColorIOColorSpaceCustomization.h
@@ -9,7 +9,7 @@
 #include "OpenColorIOColorSpace.h"
 #include "Widgets/SWidget.h"
 
-class FOpenColorIONativeConfiguration;
+class FOpenColorIOEditorConfigurationInspector;
 
 /**
  * Implements a details view customization for the FOpenColorIOConfiguration
@@ -25,7 +25,7 @@ public:
 
 protected:
 
-	static FOpenColorIONativeConfiguration* GetNativeConfig(const TSharedPtr<IPropertyHandle>& InConfigurationObjectProperty);
+	static TUniquePtr<FOpenColorIOEditorConfigurationInspector> GetConfigurationInspector(const TSharedPtr<IPropertyHandle>& InConfigurationObjectProperty);
 
 protected:
 
@@ -36,7 +36,7 @@ protected:
 	TSharedPtr<IPropertyHandle> ConfigurationObjectProperty;
 
 	/** Raw pointer to the native OCIO config. */
-	FOpenColorIONativeConfiguration* CachedNativeConfig;
+	TUniquePtr<FOpenColorIOEditorConfigurationInspector> ConfigInspector;
 };
 
 
