OCIO: working color space support to the newly added CPU processor.

We use the same approach as shaders: an extra conversion between the working color space and the interchange color space is applied either before or after the main color transform.

#rnx

diff --git a/Engine/Plugins/Compositing/OpenColorIO/Source/OpenColorIO/Private/OpenColorIOColorTransform.cpp b/Engine/Plugins/Compositing/OpenColorIO/Source/OpenColorIO/Private/OpenColorIOColorTransform.cpp
--- a/Engine/Plugins/Compositing/OpenColorIO/Source/OpenColorIO/Private/OpenColorIOColorTransform.cpp
+++ b/Engine/Plugins/Compositing/OpenColorIO/Source/OpenColorIO/Private/OpenColorIOColorTransform.cpp
@@ -612,25 +612,52 @@ bool UOpenColorIOColorTransform::IsTransform(const FString& InSourceColorSpace,
 #if WITH_EDITOR
 bool UOpenColorIOColorTransform::EditorTransformImage(const FImageView& InOutImage) const
 {
-	// TODO: Handle extra conversions needed if working color space is the source or destination
-
 #if WITH_OCIO
+	using namespace OCIO_NAMESPACE;
+
 #if !PLATFORM_EXCEPTIONS_DISABLED
 	try
 #endif
 	{
-		OCIO_NAMESPACE::ConstProcessorRcPtr TransformProcessor = GetTransformProcessor(this);
-		if (TransformProcessor)
+		ConstProcessorRcPtr Processor = GetTransformProcessor(this);
+		if (Processor && ConfigurationOwner)
 		{
-			TUniquePtr<OCIO_NAMESPACE::PackedImageDesc> ImageDesc = GetImageDesc(InOutImage);
+			TUniquePtr<PackedImageDesc> ImageDesc = GetImageDesc(InOutImage);
 			if (ImageDesc)
 			{
-				OCIO_NAMESPACE::ConstCPUProcessorRcPtr Processor = TransformProcessor->getOptimizedCPUProcessor(
-					ImageDesc->getBitDepth(),
-					ImageDesc->getBitDepth(),
-					OCIO_NAMESPACE::OptimizationFlags::OPTIMIZATION_DEFAULT
-				);
-				Processor->apply(*ImageDesc);
+				ConstConfigRcPtr InterchangeConfig = IOpenColorIOModule::Get().GetNativeInterchangeConfig_Internal()->Get();
+				ConstConfigRcPtr Config = ConfigurationOwner->GetNativeConfig_Internal()->Get();
+				BitDepth BitDepth = ImageDesc->getBitDepth();
+
+				// Conditionally apply a conversion from the working color space to interchange space
+				if (WorkingColorSpaceTransformType == EOpenColorIOWorkingColorSpaceTransform::Source)
+				{
+					ConstProcessorRcPtr	InterchangeProcessor = InterchangeConfig->GetProcessorFromConfigs(
+						InterchangeConfig,
+						StringCast<ANSICHAR>(UOpenColorIOConfiguration::WorkingColorSpaceName).Get(),
+						Config,
+						Config->getCanonicalName(OpenColorIOInterchangeName));
+
+					ConstCPUProcessorRcPtr InterchangeCPUProcessor = InterchangeProcessor->getOptimizedCPUProcessor(BitDepth, BitDepth, OPTIMIZATION_DEFAULT);
+					InterchangeCPUProcessor->apply(*ImageDesc);
+				}
+
+				// Apply the main color transformation
+				ConstCPUProcessorRcPtr CPUProcessor = Processor->getOptimizedCPUProcessor(BitDepth, BitDepth, OPTIMIZATION_DEFAULT);
+				CPUProcessor->apply(*ImageDesc);
+
+				// Conditionally apply a conversion from the interchange space to the working color space
+				if (WorkingColorSpaceTransformType == EOpenColorIOWorkingColorSpaceTransform::Destination)
+				{
+					ConstProcessorRcPtr	InterchangeProcessor = InterchangeConfig->GetProcessorFromConfigs(
+						Config,
+						Config->getCanonicalName(OpenColorIOInterchangeName),
+						InterchangeConfig,
+						StringCast<ANSICHAR>(UOpenColorIOConfiguration::WorkingColorSpaceName).Get());
+
+					ConstCPUProcessorRcPtr InterchangeCPUProcessor = InterchangeProcessor->getOptimizedCPUProcessor(BitDepth, BitDepth, OPTIMIZATION_DEFAULT);
+					InterchangeCPUProcessor->apply(*ImageDesc);
+				}
 				
 				return true;
 			}
@@ -649,26 +676,55 @@ bool UOpenColorIOColorTransform::EditorTransformImage(const FImageView& InOutIma
 
 bool UOpenColorIOColorTransform::EditorTransformImage(const FImageView& SrcImage, const FImageView& DestImage) const
 {
-	// TODO: Handle extra conversions needed if working color space is the source or destination
-
 #if WITH_OCIO
+	using namespace OCIO_NAMESPACE;
 #if !PLATFORM_EXCEPTIONS_DISABLED
 	try
 #endif
 	{
-		OCIO_NAMESPACE::ConstProcessorRcPtr TransformProcessor = GetTransformProcessor(this);
-		if (TransformProcessor)
+		ConstProcessorRcPtr Processor = GetTransformProcessor(this);
+		if (Processor && ConfigurationOwner)
 		{
-			TUniquePtr<OCIO_NAMESPACE::PackedImageDesc> SrcImageDesc = GetImageDesc(SrcImage);
-			TUniquePtr<OCIO_NAMESPACE::PackedImageDesc> DestImageDesc = GetImageDesc(DestImage);
+			TUniquePtr<PackedImageDesc> SrcImageDesc = GetImageDesc(SrcImage);
+			TUniquePtr<PackedImageDesc> DestImageDesc = GetImageDesc(DestImage);
 			if (SrcImageDesc && DestImageDesc)
 			{
-				OCIO_NAMESPACE::ConstCPUProcessorRcPtr Processor = TransformProcessor->getOptimizedCPUProcessor(
-					SrcImageDesc->getBitDepth(),
-					DestImageDesc->getBitDepth(),
-					OCIO_NAMESPACE::OptimizationFlags::OPTIMIZATION_DEFAULT
-				);
-				TransformProcessor->getDefaultCPUProcessor()->apply(*SrcImageDesc, *DestImageDesc);
+				ConstConfigRcPtr	InterchangeConfig = IOpenColorIOModule::Get().GetNativeInterchangeConfig_Internal()->Get();
+				ConstConfigRcPtr	Config = ConfigurationOwner->GetNativeConfig_Internal()->Get();
+
+				BitDepth SrcBitDepth = SrcImageDesc->getBitDepth();
+				BitDepth DestBitDepth = DestImageDesc->getBitDepth();
+
+				// Conditionally apply a conversion from the working color space to interchange space
+				if (WorkingColorSpaceTransformType == EOpenColorIOWorkingColorSpaceTransform::Source)
+				{
+					ConstProcessorRcPtr	InterchangeProcessor = InterchangeConfig->GetProcessorFromConfigs(
+						InterchangeConfig,
+						StringCast<ANSICHAR>(UOpenColorIOConfiguration::WorkingColorSpaceName).Get(),
+						Config,
+						Config->getCanonicalName(OpenColorIOInterchangeName));
+
+					ConstCPUProcessorRcPtr InterchangeCPUProcessor = InterchangeProcessor->getOptimizedCPUProcessor(SrcBitDepth, SrcBitDepth, OPTIMIZATION_DEFAULT);
+					InterchangeCPUProcessor->apply(*SrcImageDesc);
+				}
+
+
+				// Apply the main color transformation
+				ConstCPUProcessorRcPtr CPUProcessor = Processor->getOptimizedCPUProcessor(SrcBitDepth, DestBitDepth, OPTIMIZATION_DEFAULT);
+				CPUProcessor->apply(*SrcImageDesc, *DestImageDesc);
+
+				// Conditionally apply a conversion from the interchange space to the working color space
+				if (WorkingColorSpaceTransformType == EOpenColorIOWorkingColorSpaceTransform::Destination)
+				{
+					ConstProcessorRcPtr	InterchangeProcessor = InterchangeConfig->GetProcessorFromConfigs(
+						Config,
+						Config->getCanonicalName(OpenColorIOInterchangeName),
+						InterchangeConfig,
+						StringCast<ANSICHAR>(UOpenColorIOConfiguration::WorkingColorSpaceName).Get());
+
+					ConstCPUProcessorRcPtr InterchangeCPUProcessor = InterchangeProcessor->getOptimizedCPUProcessor(DestBitDepth, DestBitDepth, OPTIMIZATION_DEFAULT);
+					InterchangeCPUProcessor->apply(*DestImageDesc);
+				}
 
 				return true;
 			}
diff --git a/Engine/Plugins/Compositing/OpenColorIO/Source/OpenColorIO/Private/OpenColorIOConfiguration.cpp b/Engine/Plugins/Compositing/OpenColorIO/Source/OpenColorIO/Private/OpenColorIOConfiguration.cpp
--- a/Engine/Plugins/Compositing/OpenColorIO/Source/OpenColorIO/Private/OpenColorIOConfiguration.cpp
+++ b/Engine/Plugins/Compositing/OpenColorIO/Source/OpenColorIO/Private/OpenColorIOConfiguration.cpp
@@ -625,7 +625,7 @@ void UOpenColorIOConfiguration::LoadConfiguration()
 			FString ConfigurationFilePath = ConfigurationFile.FilePath;
 			if (ConfigurationFilePath.StartsWith(TEXT("ocio://")))
 			{
-				NativeConfig->Set(OCIO_NAMESPACE::Config::CreateFromFile(StringCast<ANSICHAR>(*ConfigurationFilePath).Get()), WorkingColorSpaceName);
+				NativeConfig->Set(OCIO_NAMESPACE::Config::CreateFromFile(StringCast<ANSICHAR>(*ConfigurationFilePath).Get()));
 				UE_LOG(LogOpenColorIO, Verbose, TEXT("Loaded built-in OCIO configuration file %s"), *ConfigurationFilePath);
 			}
 			else
@@ -650,7 +650,7 @@ void UOpenColorIOConfiguration::LoadConfiguration()
 				if (NewConfig)
 				{
 					UE_LOG(LogOpenColorIO, Verbose, TEXT("Loaded OCIO configuration file %s"), *FullPath);
-					NativeConfig->Set(NewConfig, WorkingColorSpaceName);
+					NativeConfig->Set(NewConfig);
 					
 					StartDirectoryWatch(FullPath);
 				}
diff --git a/Engine/Plugins/Compositing/OpenColorIO/Source/OpenColorIO/Private/OpenColorIOModule.cpp b/Engine/Plugins/Compositing/OpenColorIO/Source/OpenColorIO/Private/OpenColorIOModule.cpp
--- a/Engine/Plugins/Compositing/OpenColorIO/Source/OpenColorIO/Private/OpenColorIOModule.cpp
+++ b/Engine/Plugins/Compositing/OpenColorIO/Source/OpenColorIO/Private/OpenColorIOModule.cpp
@@ -5,6 +5,7 @@
 #include "Interfaces/IPluginManager.h"
 #include "OpenColorIOLibHandler.h"
 #include "OpenColorIODisplayManager.h"
+#include "OpenColorIONativeConfiguration.h"
 #include "ShaderCore.h"
 
 
@@ -20,7 +21,7 @@ FOpenColorIOModule::FOpenColorIOModule()
 
 void FOpenColorIOModule::StartupModule()
 {
-	FOpenColorIOLibHandler::Initialize();
+	bInitializedLib = FOpenColorIOLibHandler::Initialize();
 
 	// Maps virtual shader source directory /Plugin/OpenCVLensDistortion to the plugin's actual Shaders directory.
 	FString PluginShaderDir = FPaths::Combine(IPluginManager::Get().FindPlugin(TEXT("OpenColorIO"))->GetBaseDir(), TEXT("Shaders"));
@@ -36,6 +37,21 @@ FOpenColorIODisplayManager& FOpenColorIOModule::GetDisplayManager()
 {
 	return *DisplayManager;
 }
+
+FOpenColorIONativeInterchangeConfiguration* FOpenColorIOModule::GetNativeInterchangeConfig_Internal()
+{
+	if (bInitializedLib && NativeInterchangeConfig == nullptr)
+	{
+		/**
+		 * Create a single config used for conversions between the working color spaceand the interchange space.
+		 * Note that we delay the creation here to ensure that the global working color space has been loaded by
+		 * the engine via renderer settings.
+		 */
+		NativeInterchangeConfig = MakeUnique<FOpenColorIONativeInterchangeConfiguration>();
+	}
+
+	return NativeInterchangeConfig.Get();
+}
 	
 IMPLEMENT_MODULE(FOpenColorIOModule, OpenColorIO);
 
diff --git a/Engine/Plugins/Compositing/OpenColorIO/Source/OpenColorIO/Private/OpenColorIOModule.h b/Engine/Plugins/Compositing/OpenColorIO/Source/OpenColorIO/Private/OpenColorIOModule.h
--- a/Engine/Plugins/Compositing/OpenColorIO/Source/OpenColorIO/Private/OpenColorIOModule.h
+++ b/Engine/Plugins/Compositing/OpenColorIO/Source/OpenColorIO/Private/OpenColorIOModule.h
@@ -13,6 +13,7 @@ DECLARE_LOG_CATEGORY_EXTERN(LogOpenColorIO, Log, All);
 
 
 class FOpenColorIODisplayManager;
+class FOpenColorIONativeInterchangeConfiguration;
 
 /**
  * Implements the OpenColorIO module.
@@ -31,10 +32,13 @@ public:
 
 	//~ Begin IOpenColorIOModule
 	virtual FOpenColorIODisplayManager& GetDisplayManager() override;
+	virtual FOpenColorIONativeInterchangeConfiguration* GetNativeInterchangeConfig_Internal() override;
 	//~ End IOpenColorIOModule
 
 private:
 
+	bool bInitializedLib = false;
 	TUniquePtr<FOpenColorIODisplayManager> DisplayManager;
+	TUniquePtr<FOpenColorIONativeInterchangeConfiguration> NativeInterchangeConfig;
 };
 
diff --git a/Engine/Plugins/Compositing/OpenColorIO/Source/OpenColorIO/Private/OpenColorIONativeConfiguration.cpp b/Engine/Plugins/Compositing/OpenColorIO/Source/OpenColorIO/Private/OpenColorIONativeConfiguration.cpp
new file mode 100644
--- /dev/null
+++ b/Engine/Plugins/Compositing/OpenColorIO/Source/OpenColorIO/Private/OpenColorIONativeConfiguration.cpp
@@ -0,0 +1,77 @@
+// Copyright Epic Games, Inc. All Rights Reserved.
+
+#include "OpenColorIONativeConfiguration.h"
+
+#include "ColorSpace.h"
+#include "OpenColorIOConfiguration.h"
+
+
+#if WITH_EDITOR && WITH_OCIO
+OCIO_NAMESPACE::ConstConfigRcPtr FOpenColorIONativeConfiguration::Get() const
+{
+	return Config;
+}
+
+void FOpenColorIONativeConfiguration::Set(OCIO_NAMESPACE::ConstConfigRcPtr InConfig)
+{
+	using namespace OCIO_NAMESPACE;
+
+	if (InConfig == nullptr)
+	{
+		Config.reset();
+	}
+	else
+	{
+		ConstColorSpaceRcPtr InterchangeCS = InConfig->getColorSpace(InConfig->getCanonicalName(OpenColorIOInterchangeName));
+
+		// When the aces interchange color space is present, we add the working color space as an additional option.
+		if (InterchangeCS != nullptr && InConfig->getColorSpace(StringCast<ANSICHAR>(UOpenColorIOConfiguration::WorkingColorSpaceName).Get()) == nullptr)
+		{
+			ColorSpaceRcPtr WorkingCS = InterchangeCS->createEditableCopy();
+			WorkingCS->setName(StringCast<ANSICHAR>(UOpenColorIOConfiguration::WorkingColorSpaceName).Get());
+			WorkingCS->setFamily("UE");
+			WorkingCS->clearAliases();
+
+			ConfigRcPtr ConfigCopy = InConfig->createEditableCopy();
+			ConfigCopy->addColorSpace(WorkingCS);
+
+			Config = ConfigCopy;
+		}
+		else
+		{
+			Config = InConfig;
+		}
+	}
+}
+
+FOpenColorIONativeInterchangeConfiguration::FOpenColorIONativeInterchangeConfiguration()
+{
+	using namespace OCIO_NAMESPACE;
+	using namespace UE::Color;
+
+	ColorSpaceRcPtr AP0 = ColorSpace::Create();
+	AP0->setName("ACES2065-1");
+	AP0->setBitDepth(BIT_DEPTH_F32);
+	AP0->setEncoding("scene-linear");
+
+	ColorSpaceRcPtr WCS = ColorSpace::Create();
+	WCS->setName(StringCast<ANSICHAR>(UOpenColorIOConfiguration::WorkingColorSpaceName).Get());
+	WCS->setBitDepth(BIT_DEPTH_F32);
+	WCS->setEncoding("scene-linear");
+
+	const FMatrix44d TransformMat = Transpose<double>(FColorSpaceTransform(FColorSpace::GetWorking(), FColorSpace(EColorSpace::ACESAP0)));
+	MatrixTransformRcPtr MatrixTransform = MatrixTransform::Create();
+	MatrixTransform->setMatrix(&TransformMat.M[0][0]);
+	WCS->setTransform(MatrixTransform, COLORSPACE_DIR_TO_REFERENCE);
+
+	Config = OCIO_NAMESPACE::Config::Create();
+	Config->addColorSpace(AP0);
+	Config->addColorSpace(WCS);
+	Config->setRole("aces_interchange", "ACES2065-1");
+}
+
+OCIO_NAMESPACE::ConstConfigRcPtr FOpenColorIONativeInterchangeConfiguration::Get() const
+{
+	return Config;
+}
+#endif
diff --git a/Engine/Plugins/Compositing/OpenColorIO/Source/OpenColorIO/Private/OpenColorIONativeConfiguration.h b/Engine/Plugins/Compositing/OpenColorIO/Source/OpenColorIO/Private/OpenColorIONativeConfiguration.h
--- a/Engine/Plugins/Compositing/OpenColorIO/Source/OpenColorIO/Private/OpenColorIONativeConfiguration.h
+++ b/Engine/Plugins/Compositing/OpenColorIO/Source/OpenColorIO/Private/OpenColorIONativeConfiguration.h
@@ -2,54 +2,41 @@
 
 #pragma once
 
+#include "HAL/Platform.h"
+
 #if WITH_EDITOR && WITH_OCIO
 #include "OpenColorIO/OpenColorIO.h"
 #endif
 
+// Note: Currently can't delay load OCIO_NAMESPACE::ROLE_INTERCHANGE_SCENE so we define it here
+static constexpr ANSICHAR OpenColorIOInterchangeName[] = "aces_interchange";
+
 class FOpenColorIONativeConfiguration {
 public:
-
 #if WITH_EDITOR && WITH_OCIO
-	OCIO_NAMESPACE::ConstConfigRcPtr Get() const
-	{
-		return Config;
-	}
-
-	void Set(OCIO_NAMESPACE::ConstConfigRcPtr InConfig, FStringView InDefaultWorkingColorSpaceName = FStringView())
-	{
-		using namespace OCIO_NAMESPACE;
-
-		if (InConfig == nullptr)
-		{
-			Config.reset();
-		}
-		else
-		{	
-			ConstColorSpaceRcPtr InterchangeCS = InConfig->getColorSpace(InConfig->getCanonicalName("aces_interchange"/*OCIO_NAMESPACE::ROLE_INTERCHANGE_SCENE*/));
-
-			// When the aces interchange color space is present, we add the working color space as an additional option.
-			if (InterchangeCS != nullptr && InConfig->getColorSpace(StringCast<ANSICHAR>(InDefaultWorkingColorSpaceName.GetData()).Get()) == nullptr)
-			{
-				ColorSpaceRcPtr WorkingCS = InterchangeCS->createEditableCopy();
-				WorkingCS->setName(StringCast<ANSICHAR>(InDefaultWorkingColorSpaceName.GetData()).Get());
-				WorkingCS->setFamily("UE");
-				WorkingCS->clearAliases();
-
-				ConfigRcPtr ConfigCopy = InConfig->createEditableCopy();
-				ConfigCopy->addColorSpace(WorkingCS);
-
-				Config = ConfigCopy;
-			}
-			else
-			{
-				Config = InConfig;
-			}
-		}
-	}
-#endif
+	/* Native config object getter. */
+	OCIO_NAMESPACE::ConstConfigRcPtr Get() const;
+
+	/* Native config object setter. */
+	void Set(OCIO_NAMESPACE::ConstConfigRcPtr InConfig);
 
 private:
-#if WITH_EDITORONLY_DATA && WITH_OCIO
+	/* Loaded native config object. */
 	OCIO_NAMESPACE::ConstConfigRcPtr Config;
 #endif //WITH_EDITORONLY_DATA
 };
+
+
+class FOpenColorIONativeInterchangeConfiguration {
+public:
+#if WITH_EDITOR && WITH_OCIO
+	FOpenColorIONativeInterchangeConfiguration();
+
+	/* Native config object getter. */
+	OCIO_NAMESPACE::ConstConfigRcPtr Get() const;
+
+private:
+	/** Minimal config used for CPU-side conversions between the working color space and the interchange one. */
+	OCIO_NAMESPACE::ConfigRcPtr Config;
+#endif
+};
\ No newline at end of file
diff --git a/Engine/Plugins/Compositing/OpenColorIO/Source/OpenColorIO/Public/IOpenColorIOModule.h b/Engine/Plugins/Compositing/OpenColorIO/Source/OpenColorIO/Public/IOpenColorIOModule.h
--- a/Engine/Plugins/Compositing/OpenColorIO/Source/OpenColorIO/Public/IOpenColorIOModule.h
+++ b/Engine/Plugins/Compositing/OpenColorIO/Source/OpenColorIO/Public/IOpenColorIOModule.h
@@ -9,6 +9,7 @@
 
 
 class FOpenColorIODisplayManager;
+class FOpenColorIONativeInterchangeConfiguration;
 
 /**
  * Interface for the OpenColorIO module.
@@ -43,4 +44,9 @@ public:
 	 * Returns the collections of ocio display look currently in play
 	 */
 	virtual FOpenColorIODisplayManager& GetDisplayManager() = 0;
+
+	/**
+	 * Internal only: Returns a minimal dynamically-created native config for conversions between interchange and working color spaces.
+	 */
+	virtual FOpenColorIONativeInterchangeConfiguration* GetNativeInterchangeConfig_Internal() = 0;
 };
\ No newline at end of file
