OCIO: Cleanup of deprecated functions.

#rnx

diff --git a/Engine/Plugins/Compositing/OpenColorIO/Source/OpenColorIO/Private/OpenColorIOColorTransform.cpp b/Engine/Plugins/Compositing/OpenColorIO/Source/OpenColorIO/Private/OpenColorIOColorTransform.cpp
--- a/Engine/Plugins/Compositing/OpenColorIO/Source/OpenColorIO/Private/OpenColorIOColorTransform.cpp
+++ b/Engine/Plugins/Compositing/OpenColorIO/Source/OpenColorIO/Private/OpenColorIOColorTransform.cpp
@@ -198,11 +198,6 @@ void UOpenColorIOColorTransform::ProcessSerializedShaderMaps(UOpenColorIOColorTr
 	}
 }
 
-void UOpenColorIOColorTransform::GetOpenColorIOLUTKeyGuid(const FString& InProcessorIdentifier, FGuid& OutLutGuid)
-{
-	GetOpenColorIOLUTKeyGuid(InProcessorIdentifier, FName(), OutLutGuid);
-}
-
 void UOpenColorIOColorTransform::GetOpenColorIOLUTKeyGuid(const FString& InProcessorIdentifier, const FName& InName, FGuid& OutLutGuid)
 {
 #if WITH_EDITOR
@@ -577,18 +572,6 @@ bool UOpenColorIOColorTransform::AreRenderResourcesReady() const
 	return true;
 }
 
-bool UOpenColorIOColorTransform::GetShaderAndLUTResouces(ERHIFeatureLevel::Type InFeatureLevel, FOpenColorIOTransformResource*& OutShaderResource, FTextureResource*& OutLUT3dResource)
-{
-	TSortedMap<int32, FTextureResource*> OutTextureResources;
-	bool bResult = GetRenderResources(InFeatureLevel, OutShaderResource, OutTextureResources);
-
-	if (OutTextureResources.Contains(0))
-	{
-		OutLUT3dResource = OutTextureResources[0];
-	}
-	return bResult;
-}
-
 bool UOpenColorIOColorTransform::IsTransform(const FString& InSourceColorSpace, const FString& InDestinationColorSpace) const
 {
 	if (!bIsDisplayViewType)
@@ -958,16 +941,6 @@ TObjectPtr<UTexture> UOpenColorIOColorTransform::CreateTexture3DLUT(const FStrin
 	return OutTexture;
 }
 
-void UOpenColorIOColorTransform::Update3dLutTexture(const FString& InLutIdentifier, const float* InSourceData)
-{
-#if WITH_EDITOR && WITH_OCIO
-	Textures.Empty();
-
-	TObjectPtr<UTexture> Texture = CreateTexture3DLUT(InLutIdentifier, FName(), OpenColorIOShader::Lut3dEdgeLength, TextureFilter::TF_Bilinear, InSourceData);
-	Textures.Add(0, MoveTemp(Texture));
-#endif
-}
-
 TObjectPtr<UTexture> UOpenColorIOColorTransform::CreateTexture1DLUT(const FString& InProcessorIdentifier, const FName& InName, uint32 InTextureWidth, uint32 InTextureHeight, TextureFilter InFilter, bool bRedChannelOnly, const float* InSourceData)
 {
 	TObjectPtr<UTexture2D> OutTexture = nullptr;
diff --git a/Engine/Plugins/Compositing/OpenColorIO/Source/OpenColorIO/Private/OpenColorIOConfiguration.cpp b/Engine/Plugins/Compositing/OpenColorIO/Source/OpenColorIO/Private/OpenColorIOConfiguration.cpp
--- a/Engine/Plugins/Compositing/OpenColorIO/Source/OpenColorIO/Private/OpenColorIOConfiguration.cpp
+++ b/Engine/Plugins/Compositing/OpenColorIO/Source/OpenColorIO/Private/OpenColorIOConfiguration.cpp
@@ -81,24 +81,6 @@ bool UOpenColorIOConfiguration::GetRenderResources(ERHIFeatureLevel::Type InFeat
 	return false;
 }
 
-bool UOpenColorIOConfiguration::GetShaderAndLUTResources(ERHIFeatureLevel::Type InFeatureLevel, const FString& InSourceColorSpace, const FString& InDestinationColorSpace, FOpenColorIOTransformResource*& OutShaderResource, FTextureResource*& OutLUT3dResource)
-{
-	TSortedMap<int32, FTextureResource*> TextureResources;
-	FOpenColorIOColorConversionSettings Settings;
-	Settings.SourceColorSpace.ColorSpaceName = InSourceColorSpace;
-	Settings.DestinationColorSpace.ColorSpaceName = InDestinationColorSpace;
-	if (GetRenderResources(InFeatureLevel, Settings, OutShaderResource, TextureResources))
-	{
-		if (TextureResources.Contains(0))
-		{
-			OutLUT3dResource = TextureResources[0];
-			return true;
-		}
-	}
-
-	return false;
-}
-
 bool UOpenColorIOConfiguration::HasTransform(const FString& InSourceColorSpace, const FString& InDestinationColorSpace)
 {
 	TObjectPtr<UOpenColorIOColorTransform>* TransformData = ColorTransforms.FindByPredicate([&](const UOpenColorIOColorTransform* InTransformData)
diff --git a/Engine/Plugins/Compositing/OpenColorIO/Source/OpenColorIO/Public/OpenColorIOColorTransform.h b/Engine/Plugins/Compositing/OpenColorIO/Source/OpenColorIO/Public/OpenColorIOColorTransform.h
--- a/Engine/Plugins/Compositing/OpenColorIO/Source/OpenColorIO/Public/OpenColorIOColorTransform.h
+++ b/Engine/Plugins/Compositing/OpenColorIO/Source/OpenColorIO/Public/OpenColorIOColorTransform.h
@@ -74,9 +74,6 @@ public:
 	 */
 	bool AreRenderResourcesReady() const;
 
-	UE_DEPRECATED(5.1, "GetShaderAndLUTResouces is deprecated, please use GetRenderResources instead.")
-	bool GetShaderAndLUTResouces(ERHIFeatureLevel::Type InFeatureLevel, FOpenColorIOTransformResource*& OutShaderResource, FTextureResource*& OutLUT3dResource);
-
 	bool IsTransform(const FString& InSourceColorSpace, const FString& InDestinationColorSpace) const;
 	bool IsTransform(const FString& InSourceColorSpace, const FString& InDisplay, const FString& InView, EOpenColorIOViewTransformDirection InDirection) const;
 
@@ -120,9 +117,6 @@ protected:
 	 */
 	static void GetOpenColorIOLUTKeyGuid(const FString& InProcessorIdentifier, const FName& InName, FGuid& OutLutGuid );
 
-	UE_DEPRECATED(5.1, "This version of GetOpenColorIOLUTKeyGuid is deprecated. Please use the version that takes InName instead.")
-	static void GetOpenColorIOLUTKeyGuid(const FString& InLutIdentifier, FGuid& OutLutGuid);
-
 	/**
 	 * Generate the LUT and shader associated to the desired color space transform.
 	 * @note: Only in editor
@@ -148,9 +142,6 @@ protected:
 	 */
 	TObjectPtr<UTexture> CreateTexture3DLUT(const FString& InProcessorIdentifier, const FName& InName, uint32 InLutLength, TextureFilter InFilter, const float* InSourceData);
 
-	UE_DEPRECATED(5.1, "Update3dLutTexture is deprecated, please use CreateTexture3DLUT instead.")
-	void Update3dLutTexture(const FString& InLutIdentifier, const float* InSourceData);
-
 	/**
 	 * Helper function taking raw 1D LUT data coming from the library and initializing a UTexture with it.
 	 */
diff --git a/Engine/Plugins/Compositing/OpenColorIO/Source/OpenColorIO/Public/OpenColorIOConfiguration.h b/Engine/Plugins/Compositing/OpenColorIO/Source/OpenColorIO/Public/OpenColorIOConfiguration.h
--- a/Engine/Plugins/Compositing/OpenColorIO/Source/OpenColorIO/Public/OpenColorIOConfiguration.h
+++ b/Engine/Plugins/Compositing/OpenColorIO/Source/OpenColorIO/Public/OpenColorIOConfiguration.h
@@ -41,8 +41,6 @@ public:
 
 	bool IsTransformReady(const FOpenColorIOColorConversionSettings& InSettings);
 	bool GetRenderResources(ERHIFeatureLevel::Type InFeatureLevel, const FOpenColorIOColorConversionSettings& InSettings, FOpenColorIOTransformResource*& OutShaderResource, TSortedMap<int32, FTextureResource*>& OutTextureResources);
-	UE_DEPRECATED(5.1, "GetShaderAndLUTResources is deprecated, please use GetRenderResources instead.")
-	bool GetShaderAndLUTResources(ERHIFeatureLevel::Type InFeatureLevel, const FString& InSourceColorSpace, const FString& InDestinationColorSpace, FOpenColorIOTransformResource*& OutShaderResource, FTextureResource*& OutLUT3dResource);
 	bool HasTransform(const FString& InSourceColorSpace, const FString& InDestinationColorSpace);
 	bool HasTransform(const FString& InSourceColorSpace, const FString& InDisplay, const FString& InView, EOpenColorIOViewTransformDirection InDirection);
 	bool HasDesiredColorSpace(const FOpenColorIOColorSpace& ColorSpace) const;
