nDisplay: Fix OpenColorIO display-view not working on nodes.

diff --git a/Engine/Plugins/Compositing/OpenColorIO/Source/OpenColorIO/Private/OpenColorIOColorSpace.cpp b/Engine/Plugins/Compositing/OpenColorIO/Source/OpenColorIO/Private/OpenColorIOColorSpace.cpp
--- a/Engine/Plugins/Compositing/OpenColorIO/Source/OpenColorIO/Private/OpenColorIOColorSpace.cpp
+++ b/Engine/Plugins/Compositing/OpenColorIO/Source/OpenColorIO/Private/OpenColorIOColorSpace.cpp
@@ -143,20 +143,26 @@ bool FOpenColorIOColorConversionSettings::IsValid() const
 	{
 		if (IsDisplayView())
 		{
-			switch (DisplayViewDirection)
+			if (SourceColorSpace.IsValid() && DestinationDisplayView.IsValid())
 			{
-			case EOpenColorIOViewTransformDirection::Forward:
-				return ConfigurationSource->HasTransform(SourceColorSpace.ColorSpaceName, DestinationDisplayView.Display, DestinationDisplayView.View, EOpenColorIOViewTransformDirection::Forward);
-			case EOpenColorIOViewTransformDirection::Inverse:
-				return ConfigurationSource->HasTransform(SourceColorSpace.ColorSpaceName, DestinationDisplayView.Display, DestinationDisplayView.View, EOpenColorIOViewTransformDirection::Inverse);
-			default:
-				checkNoEntry();
-				return false;
+				switch (DisplayViewDirection)
+				{
+				case EOpenColorIOViewTransformDirection::Forward:
+					return ConfigurationSource->HasTransform(SourceColorSpace.ColorSpaceName, DestinationDisplayView.Display, DestinationDisplayView.View, EOpenColorIOViewTransformDirection::Forward);
+				case EOpenColorIOViewTransformDirection::Inverse:
+					return ConfigurationSource->HasTransform(SourceColorSpace.ColorSpaceName, DestinationDisplayView.Display, DestinationDisplayView.View, EOpenColorIOViewTransformDirection::Inverse);
+				default:
+					checkNoEntry();
+					return false;
+				}
 			}
 		}
 		else
 		{
-			return ConfigurationSource->HasTransform(SourceColorSpace.ColorSpaceName, DestinationColorSpace.ColorSpaceName);
+			if (SourceColorSpace.IsValid() && DestinationColorSpace.IsValid())
+			{
+				return ConfigurationSource->HasTransform(SourceColorSpace.ColorSpaceName, DestinationColorSpace.ColorSpaceName);
+			}
 		}
 	}
 
diff --git a/Engine/Plugins/Runtime/nDisplay/Source/DisplayCluster/Private/Render/Viewport/Configuration/DisplayClusterViewportConfigurationHelpers_OpenColorIO.cpp b/Engine/Plugins/Runtime/nDisplay/Source/DisplayCluster/Private/Render/Viewport/Configuration/DisplayClusterViewportConfigurationHelpers_OpenColorIO.cpp
--- a/Engine/Plugins/Runtime/nDisplay/Source/DisplayCluster/Private/Render/Viewport/Configuration/DisplayClusterViewportConfigurationHelpers_OpenColorIO.cpp
+++ b/Engine/Plugins/Runtime/nDisplay/Source/DisplayCluster/Private/Render/Viewport/Configuration/DisplayClusterViewportConfigurationHelpers_OpenColorIO.cpp
@@ -88,9 +88,7 @@ bool FDisplayClusterViewportConfigurationHelpers_OpenColorIO::UpdateChromakeyVie
 void FDisplayClusterViewportConfigurationHelpers_OpenColorIO::ApplyOCIOConfiguration(FDisplayClusterViewport& DstViewport, const FOpenColorIOColorConversionSettings& InConfiguration)
 {
 	// Check OCIO before apply
-	if(!InConfiguration.IsValid()
-	|| !InConfiguration.SourceColorSpace.IsValid()
-	|| !InConfiguration.DestinationColorSpace.IsValid())
+	if(!InConfiguration.IsValid())
 	{
 		DisableOCIOConfiguration(DstViewport);
 
