OCIO: Added missing rendering pass information.

#rnx

diff --git a/Engine/Plugins/Compositing/OpenColorIO/Source/OpenColorIO/Private/OpenColorIODisplayExtension.cpp b/Engine/Plugins/Compositing/OpenColorIO/Source/OpenColorIO/Private/OpenColorIODisplayExtension.cpp
--- a/Engine/Plugins/Compositing/OpenColorIO/Source/OpenColorIO/Private/OpenColorIODisplayExtension.cpp
+++ b/Engine/Plugins/Compositing/OpenColorIO/Source/OpenColorIO/Private/OpenColorIODisplayExtension.cpp
@@ -55,12 +55,13 @@ void FOpenColorIODisplayExtension::SetupView(FSceneViewFamily& InViewFamily, FSc
 	//If something fails, cache invalid resources to invalidate them
 	FOpenColorIOTransformResource* ShaderResource = nullptr;
 	TSortedMap<int32, FTextureResource*> TransformTextureResources;
+	const FOpenColorIOColorConversionSettings& ConversionSettings = DisplayConfiguration.ColorConfiguration;
 
-	if (DisplayConfiguration.ColorConfiguration.ConfigurationSource != nullptr)
+	if (ConversionSettings.ConfigurationSource != nullptr)
 	{
-		const bool bFoundTransform = DisplayConfiguration.ColorConfiguration.ConfigurationSource->GetRenderResources(
+		const bool bFoundTransform = ConversionSettings.ConfigurationSource->GetRenderResources(
 			InViewFamily.GetFeatureLevel()
-			, DisplayConfiguration.ColorConfiguration
+			, ConversionSettings
 			, ShaderResource
 			, TransformTextureResources);
 
@@ -89,11 +90,12 @@ void FOpenColorIODisplayExtension::SetupView(FSceneViewFamily& InViewFamily, FSc
 	}
 
 	ENQUEUE_RENDER_COMMAND(ProcessColorSpaceTransform)(
-		[this, ShaderResource, TextureResources = MoveTemp(TransformTextureResources)](FRHICommandListImmediate& RHICmdList)
+		[this, ShaderResource, TextureResources = MoveTemp(TransformTextureResources), TransformName = ConversionSettings.ToString()](FRHICommandListImmediate& RHICmdList)
 		{
 			//Caches render thread resource to be used when applying configuration in PostRenderViewFamily_RenderThread
 			CachedResourcesRenderThread.ShaderResource = ShaderResource;
 			CachedResourcesRenderThread.TextureResources = TextureResources;
+			CachedResourcesRenderThread.TransformName = TransformName;
 		}
 	);
 }
diff --git a/Engine/Plugins/Compositing/OpenColorIO/Source/OpenColorIO/Private/OpenColorIORendering.cpp b/Engine/Plugins/Compositing/OpenColorIO/Source/OpenColorIO/Private/OpenColorIORendering.cpp
--- a/Engine/Plugins/Compositing/OpenColorIO/Source/OpenColorIO/Private/OpenColorIORendering.cpp
+++ b/Engine/Plugins/Compositing/OpenColorIO/Source/OpenColorIO/Private/OpenColorIORendering.cpp
@@ -82,8 +82,9 @@ void FOpenColorIORendering::AddPass_RenderThread(
 		}
 		Parameters->Gamma = InGamma;
 		Parameters->RenderTargets[0] = Output.GetRenderTargetBinding();
-
-		AddDrawScreenPass(GraphBuilder, RDG_EVENT_NAME("OpenColorIOPass"), ViewInfo, OutputViewport, InputViewport, VertexShader, OCIOPixelShader, Parameters);
+		
+		FRDGEventName PassName = RDG_EVENT_NAME("OpenColorIOPass %dx%d (%s)", Output.ViewRect.Width(), Output.ViewRect.Height(), InPassResource.TransformName.IsEmpty() ? TEXT("Unspecified Transform") : *InPassResource.TransformName);
+		AddDrawScreenPass(GraphBuilder, MoveTemp(PassName), ViewInfo, OutputViewport, InputViewport, VertexShader, OCIOPixelShader, Parameters);
 	}
 	else
 	{
@@ -95,7 +96,8 @@ void FOpenColorIORendering::AddPass_RenderThread(
 		Parameters->MiniFontTexture = OpenColorIOGetMiniFontTexture();
 		Parameters->RenderTargets[0] = Output.GetRenderTargetBinding();
 
-		AddDrawScreenPass(GraphBuilder, RDG_EVENT_NAME("OpenColorIOInvalidPass"), ViewInfo, OutputViewport, InputViewport, VertexShader, OCIOInvalidPixelShader, Parameters);
+		FRDGEventName PassName = RDG_EVENT_NAME("OpenColorIOInvalidPass %dx%d", Output.ViewRect.Width(), Output.ViewRect.Height());
+		AddDrawScreenPass(GraphBuilder, MoveTemp(PassName), ViewInfo, OutputViewport, InputViewport, VertexShader, OCIOInvalidPixelShader, Parameters);
 	}
 }
 
@@ -148,7 +150,7 @@ bool FOpenColorIORendering::ApplyColorTransform(UWorld* InWorld, const FOpenColo
 	}
 	
 	ENQUEUE_RENDER_COMMAND(ProcessColorSpaceTransform)(
-		[FeatureLevel, InputResource, OutputResource, ShaderResource, TextureResources = MoveTemp(TransformTextureResources)](FRHICommandListImmediate& RHICmdList)
+		[FeatureLevel, InputResource, OutputResource, ShaderResource, TextureResources = MoveTemp(TransformTextureResources), TransformName = InSettings.ToString()](FRHICommandListImmediate& RHICmdList)
 		{
 			FRDGBuilder GraphBuilder(RHICmdList);
 
@@ -163,7 +165,7 @@ bool FOpenColorIORendering::ApplyColorTransform(UWorld* InWorld, const FOpenColo
 				FeatureLevel,
 				FScreenPassTexture(InputTexture),
 				Output,
-				FOpenColorIORenderPassResources{ ShaderResource, TextureResources},
+				FOpenColorIORenderPassResources{ShaderResource, TextureResources, TransformName},
 				1.0f); // Set Gamma to 1., since we do not have any display parameters or requirement for Gamma.
 
 			GraphBuilder.Execute();
diff --git a/Engine/Plugins/Compositing/OpenColorIO/Source/OpenColorIO/Public/OpenColorIORendering.h b/Engine/Plugins/Compositing/OpenColorIO/Source/OpenColorIO/Public/OpenColorIORendering.h
--- a/Engine/Plugins/Compositing/OpenColorIO/Source/OpenColorIO/Public/OpenColorIORendering.h
+++ b/Engine/Plugins/Compositing/OpenColorIO/Source/OpenColorIO/Public/OpenColorIORendering.h
@@ -17,10 +17,17 @@ class UTexture;
 class UTextureRenderTarget2D;
 class UWorld;
 
+/** Resources needed by the FOpenColorIORendering pass function. */
 struct OPENCOLORIO_API FOpenColorIORenderPassResources
 {
+	/** Color transform pass (generated) shader. */
 	FOpenColorIOTransformResource* ShaderResource = nullptr;
+	
+	/** Collection of LUT textures needed by the shader. */
 	TSortedMap<int32, FTextureResource*> TextureResources = {};
+
+	/** Color transform string description. */
+	FString TransformName = FString();
 };
 
 /** Entry point to trigger OpenColorIO conversion rendering */
diff --git a/Engine/Plugins/Runtime/nDisplay/Source/DisplayCluster/Private/Render/Viewport/DisplayClusterViewport_OpenColorIO.cpp b/Engine/Plugins/Runtime/nDisplay/Source/DisplayCluster/Private/Render/Viewport/DisplayClusterViewport_OpenColorIO.cpp
--- a/Engine/Plugins/Runtime/nDisplay/Source/DisplayCluster/Private/Render/Viewport/DisplayClusterViewport_OpenColorIO.cpp
+++ b/Engine/Plugins/Runtime/nDisplay/Source/DisplayCluster/Private/Render/Viewport/DisplayClusterViewport_OpenColorIO.cpp
@@ -94,11 +94,12 @@ void FDisplayClusterViewport_OpenColorIO::UpdateOpenColorIORenderPassResources()
 	}
 
 	ENQUEUE_RENDER_COMMAND(ProcessColorSpaceTransform)(
-		[This = SharedThis(this), ShaderResource, TextureResources = MoveTemp(TransformTextureResources)](FRHICommandListImmediate& RHICmdList)
+		[This = SharedThis(this), ShaderResource, TextureResources = MoveTemp(TransformTextureResources), TransformName = ConversionSettings.ToString()](FRHICommandListImmediate& RHICmdList)
 		{
 			//Caches render thread resource to be used when applying configuration in PostRenderViewFamily_RenderThread
 			This->CachedResourcesRenderThread.ShaderResource = ShaderResource;
 			This->CachedResourcesRenderThread.TextureResources = TextureResources;
+			This->CachedResourcesRenderThread.TransformName = TransformName;
 		}
 	);
 }
