OCIO: Expose option to apply color transform without premultiplication & alpha inversion, fixing ndisplay lightcards.

diff --git a/Engine/Plugins/Compositing/OpenColorIO/Shaders/Private/OpenColorIOShader.usf b/Engine/Plugins/Compositing/OpenColorIO/Shaders/Private/OpenColorIOShader.usf
--- a/Engine/Plugins/Compositing/OpenColorIO/Shaders/Private/OpenColorIOShader.usf
+++ b/Engine/Plugins/Compositing/OpenColorIO/Shaders/Private/OpenColorIOShader.usf
@@ -13,6 +13,7 @@ SamplerState	InputTextureSampler;
 float4x4		WorkingColorSpaceToInterchange;
 float4x4		InterchangeToWorkingColorSpace;
 float			Gamma;
+uint TransformAlpha;
 
 /*=============================================================================
  * OCIO generated start
@@ -28,6 +29,29 @@ float			Gamma;
 /* Helpers
 =============================================================================*/
 
+float SafeRcp(float x)
+{
+	if (x > 0)
+	{
+		return rcp(x);
+	}
+	return 0;
+}
+
+float4 InvertAlpha(float4 InColor)
+{
+	return float4(InColor.xyz, 1.0 - InColor.w);
+}
+
+float4 Unpremultiply(float4 InColor)
+{
+	return float4(InColor.xyz * SafeRcp(InColor.w), InColor.w);
+}
+
+float4 Premultiply(float4 InColor)
+{
+	return float4(InColor.xyz * InColor.w, InColor.w);
+}
 
 
 /* Pixel shader
@@ -40,6 +64,16 @@ void MainPS(in noperspective float4 InUVAndScreenPos : TEXCOORD0,
 	InputColor.xyz = max(0.0, InputColor.xyz);
 	InputColor.xyz = pow(InputColor.xyz, Gamma);
 
+	if(TransformAlpha > 0)
+	{
+		if(TransformAlpha == 2 /*InvertUnpremultiply*/)
+		{
+			InputColor = InvertAlpha(InputColor);
+		}
+
+		InputColor = Unpremultiply(InputColor);
+	}
+
 #if WORKING_COLOR_SPACE_TRANSFORM_TYPE == 1 // EOpenColorIOWorkingColorSpaceTransform::Source
 	InputColor.xyz = mul((float3x3)WorkingColorSpaceToInterchange, InputColor.xyz);
 #endif
@@ -50,4 +84,14 @@ void MainPS(in noperspective float4 InUVAndScreenPos : TEXCOORD0,
 #if WORKING_COLOR_SPACE_TRANSFORM_TYPE == 2 // EOpenColorIOWorkingColorSpaceTransform::Destination
 	OutColor.xyz = mul((float3x3)InterchangeToWorkingColorSpace, OutColor.xyz);
 #endif
+
+	if(TransformAlpha > 0)
+	{
+		OutColor = Premultiply(OutColor);
+		
+		if(TransformAlpha == 2 /*InvertUnpremultiply*/)
+		{
+			OutColor = InvertAlpha(OutColor);
+		}
+	}
 }
\ No newline at end of file
diff --git a/Engine/Plugins/Compositing/OpenColorIO/Source/OpenColorIO/Private/OpenColorIORendering.cpp b/Engine/Plugins/Compositing/OpenColorIO/Source/OpenColorIO/Private/OpenColorIORendering.cpp
--- a/Engine/Plugins/Compositing/OpenColorIO/Source/OpenColorIO/Private/OpenColorIORendering.cpp
+++ b/Engine/Plugins/Compositing/OpenColorIO/Source/OpenColorIO/Private/OpenColorIORendering.cpp
@@ -47,7 +47,8 @@ void FOpenColorIORendering::AddPass_RenderThread(
 	const FScreenPassTexture& Input,
 	const FScreenPassRenderTarget& Output,
 	const FOpenColorIORenderPassResources& InPassResource,
-	float InGamma)
+	float InGamma, 
+	EOpenColorIOTransformAlpha TransformAlpha)
 {
 	check(IsInRenderingThread());
 
@@ -81,6 +82,7 @@ void FOpenColorIORendering::AddPass_RenderThread(
 			break;	
 		}
 		Parameters->Gamma = InGamma;
+		Parameters->TransformAlpha = (uint32)TransformAlpha;
 		Parameters->RenderTargets[0] = Output.GetRenderTargetBinding();
 		
 		FRDGEventName PassName = RDG_EVENT_NAME("OpenColorIOPass %dx%d (%s)", Output.ViewRect.Width(), Output.ViewRect.Height(), InPassResource.TransformName.IsEmpty() ? TEXT("Unspecified Transform") : *InPassResource.TransformName);
diff --git a/Engine/Plugins/Compositing/OpenColorIO/Source/OpenColorIO/Public/OpenColorIORendering.h b/Engine/Plugins/Compositing/OpenColorIO/Source/OpenColorIO/Public/OpenColorIORendering.h
--- a/Engine/Plugins/Compositing/OpenColorIO/Source/OpenColorIO/Public/OpenColorIORendering.h
+++ b/Engine/Plugins/Compositing/OpenColorIO/Source/OpenColorIO/Public/OpenColorIORendering.h
@@ -30,6 +30,14 @@ struct OPENCOLORIO_API FOpenColorIORenderPassResources
 	FString TransformName = FString();
 };
 
+/** Option to transform the render target alpha before applying the color transform. */
+enum class EOpenColorIOTransformAlpha : uint32
+{
+	None = 0,
+	Unpremultiply = 1,
+	InvertUnpremultiply = 2
+};
+
 /** Entry point to trigger OpenColorIO conversion rendering */
 class OPENCOLORIO_API FOpenColorIORendering
 {
@@ -56,6 +64,7 @@ public:
 	 * @param Output Destination render target
 	 * @param InPassInfo OpenColorIO shader and texture resources
 	 * @param InGamma Display gamma
+	 * @param TransformAlpha Whether to unpremult/invert before applying the color transform
 	 */
 	static void AddPass_RenderThread(
 		FRDGBuilder& GraphBuilder,
@@ -64,5 +73,6 @@ public:
 		const FScreenPassTexture& Input,
 		const FScreenPassRenderTarget& Output,
 		const FOpenColorIORenderPassResources& InPassInfo,
-		float InGamma);
+		float InGamma,
+		EOpenColorIOTransformAlpha TransformAlpha = EOpenColorIOTransformAlpha::None);
 };
\ No newline at end of file
diff --git a/Engine/Plugins/Compositing/OpenColorIO/Source/OpenColorIO/Public/OpenColorIOShader.h b/Engine/Plugins/Compositing/OpenColorIO/Source/OpenColorIO/Public/OpenColorIOShader.h
--- a/Engine/Plugins/Compositing/OpenColorIO/Source/OpenColorIO/Public/OpenColorIOShader.h
+++ b/Engine/Plugins/Compositing/OpenColorIO/Source/OpenColorIO/Public/OpenColorIOShader.h
@@ -23,6 +23,7 @@ BEGIN_SHADER_PARAMETER_STRUCT(FOpenColorIOPixelShaderParameters, )
 	SHADER_PARAMETER(FMatrix44f, WorkingColorSpaceToInterchange)
 	SHADER_PARAMETER(FMatrix44f, InterchangeToWorkingColorSpace)
 	SHADER_PARAMETER(float, Gamma)
+	SHADER_PARAMETER(uint32, TransformAlpha)
 	
 	SHADER_PARAMETER_TEXTURE(Texture3D, Ocio_lut3d_0)
 	SHADER_PARAMETER_SAMPLER(SamplerState, Ocio_lut3d_0Sampler)
diff --git a/Engine/Plugins/Runtime/nDisplay/Source/DisplayCluster/Private/Render/Viewport/DisplayClusterViewportProxy.cpp b/Engine/Plugins/Runtime/nDisplay/Source/DisplayCluster/Private/Render/Viewport/DisplayClusterViewportProxy.cpp
--- a/Engine/Plugins/Runtime/nDisplay/Source/DisplayCluster/Private/Render/Viewport/DisplayClusterViewportProxy.cpp
+++ b/Engine/Plugins/Runtime/nDisplay/Source/DisplayCluster/Private/Render/Viewport/DisplayClusterViewportProxy.cpp
@@ -680,11 +680,16 @@ void FDisplayClusterViewportProxy::UpdateDeferredResources(FRHICommandListImmedi
 
 			for (int32 ContextNum = 0; ContextNum < Input.Num(); ContextNum++)
 			{
+				const bool bUnpremultiply = EnumHasAnyFlags(RenderSettingsICVFX.RuntimeFlags, EDisplayClusterViewportRuntimeICVFXFlags::UVLightcard | EDisplayClusterViewportRuntimeICVFXFlags::Lightcard);
+				const bool bInvertAlpha = !EnumHasAnyFlags(RenderSettingsICVFX.RuntimeFlags, EDisplayClusterViewportRuntimeICVFXFlags::UVLightcard);
+
 				if (OpenColorIO->AddPass_RenderThread(
 					GraphBuilder,
 					SourceViewportProxy.GetContexts_RenderThread()[ContextNum],
 					Input[ContextNum], InputRects[ContextNum],
-					Output[ContextNum], OutputRects[ContextNum]))
+					Output[ContextNum], OutputRects[ContextNum],
+					bUnpremultiply,
+					bInvertAlpha))
 				{
 					bPass0Applied = true;
 				}
diff --git a/Engine/Plugins/Runtime/nDisplay/Source/DisplayCluster/Private/Render/Viewport/DisplayClusterViewport_OpenColorIO.cpp b/Engine/Plugins/Runtime/nDisplay/Source/DisplayCluster/Private/Render/Viewport/DisplayClusterViewport_OpenColorIO.cpp
--- a/Engine/Plugins/Runtime/nDisplay/Source/DisplayCluster/Private/Render/Viewport/DisplayClusterViewport_OpenColorIO.cpp
+++ b/Engine/Plugins/Runtime/nDisplay/Source/DisplayCluster/Private/Render/Viewport/DisplayClusterViewport_OpenColorIO.cpp
@@ -127,7 +127,7 @@ float FDisplayClusterViewport_OpenColorIO::GetDisplayGamma(const FDisplayCluster
 }
 
 bool FDisplayClusterViewport_OpenColorIO::AddPass_RenderThread(FRDGBuilder& GraphBuilder, const FDisplayClusterViewport_Context& InViewportContext,
-	FRHITexture2D* InputTextureRHI, const FIntRect& InputRect, FRHITexture2D* OutputTextureRHI, const FIntRect& OutputRect) const
+	FRHITexture2D* InputTextureRHI, const FIntRect& InputRect, FRHITexture2D* OutputTextureRHI, const FIntRect& OutputRect, bool bUnpremultiply, bool bInvertAlpha) const
 {
 	if (IsEnabled_RenderThread())
 	{
@@ -135,6 +135,11 @@ bool FDisplayClusterViewport_OpenColorIO::AddPass_RenderThread(FRDGBuilder& Grap
 		FRDGTextureRef OutputTexture = GraphBuilder.RegisterExternalTexture(CreateRenderTarget(OutputTextureRHI, TEXT("nDisplayOCIORenderTargetTexture")));
 
 		FScreenPassRenderTarget Output = FScreenPassRenderTarget(OutputTexture, OutputRect, ERenderTargetLoadAction::EClear);
+		EOpenColorIOTransformAlpha TransformAlpha = EOpenColorIOTransformAlpha::None;
+		if (bUnpremultiply)
+		{
+			TransformAlpha = bInvertAlpha ? EOpenColorIOTransformAlpha::InvertUnpremultiply: EOpenColorIOTransformAlpha::Unpremultiply;
+		}
 
 		FOpenColorIORendering::AddPass_RenderThread(
 			GraphBuilder,
@@ -143,7 +148,8 @@ bool FDisplayClusterViewport_OpenColorIO::AddPass_RenderThread(FRDGBuilder& Grap
 			FScreenPassTexture(InputTexture),
 			Output,
 			CachedResourcesRenderThread,
-			GetDisplayGamma(InViewportContext)
+			GetDisplayGamma(InViewportContext),
+			TransformAlpha
 		);
 
 		return true;
diff --git a/Engine/Plugins/Runtime/nDisplay/Source/DisplayCluster/Private/Render/Viewport/DisplayClusterViewport_OpenColorIO.h b/Engine/Plugins/Runtime/nDisplay/Source/DisplayCluster/Private/Render/Viewport/DisplayClusterViewport_OpenColorIO.h
--- a/Engine/Plugins/Runtime/nDisplay/Source/DisplayCluster/Private/Render/Viewport/DisplayClusterViewport_OpenColorIO.h
+++ b/Engine/Plugins/Runtime/nDisplay/Source/DisplayCluster/Private/Render/Viewport/DisplayClusterViewport_OpenColorIO.h
@@ -55,7 +55,7 @@ public:
 	 * @return - true if success.
 	 */
 	bool AddPass_RenderThread(FRDGBuilder& GraphBuilder, const FDisplayClusterViewport_Context& InViewportContext,
-		FRHITexture2D* InputTextureRHI, const FIntRect& InputRect, FRHITexture2D* OutputTextureRHI, const FIntRect& OutputRect) const;
+		FRHITexture2D* InputTextureRHI, const FIntRect& InputRect, FRHITexture2D* OutputTextureRHI, const FIntRect& OutputRect, bool bUnpremultiply, bool InvertedAlpha) const;
 
 	/* This is a copy of FOpenColorIODisplayExtension::PostProcessPassAfterTonemap_RenderThread() */
 	FScreenPassTexture PostProcessPassAfterTonemap_RenderThread(FRDGBuilder& GraphBuilder, const FDisplayClusterViewport_Context& InViewportContext, const FSceneView& View, const FPostProcessMaterialInputs& InOutInputs);
