OCIO: Make WITH_OCIO by definition editor-only.

#rnx

diff --git a/Engine/Plugins/Compositing/OpenColorIO/Source/OpenColorIO/OpenColorIO.Build.cs b/Engine/Plugins/Compositing/OpenColorIO/Source/OpenColorIO/OpenColorIO.Build.cs
--- a/Engine/Plugins/Compositing/OpenColorIO/Source/OpenColorIO/OpenColorIO.Build.cs
+++ b/Engine/Plugins/Compositing/OpenColorIO/Source/OpenColorIO/OpenColorIO.Build.cs
@@ -57,6 +57,10 @@ namespace UnrealBuildTool.Rules
 						"UnrealEd"
 					});
 			}
+			else
+			{
+				PrivateDefinitions.Add("WITH_OCIO=0");
+			}
 		}
 	}
 }
diff --git a/Engine/Plugins/Compositing/OpenColorIO/Source/OpenColorIO/Private/OpenColorIOColorTransform.cpp b/Engine/Plugins/Compositing/OpenColorIO/Source/OpenColorIO/Private/OpenColorIOColorTransform.cpp
--- a/Engine/Plugins/Compositing/OpenColorIO/Source/OpenColorIO/Private/OpenColorIOColorTransform.cpp
+++ b/Engine/Plugins/Compositing/OpenColorIO/Source/OpenColorIO/Private/OpenColorIOColorTransform.cpp
@@ -31,7 +31,7 @@
 #endif //WITH_EDITOR
 
 namespace {
-#if WITH_EDITOR && WITH_OCIO
+#if WITH_OCIO
 	/*
 	 * Get the processor optimization flag.
 	 */
@@ -333,7 +333,7 @@ void UOpenColorIOColorTransform::CacheResourceTextures()
 {
 	if (Textures.IsEmpty())
 	{
-#if WITH_EDITOR && WITH_OCIO
+#if WITH_OCIO
 #if !PLATFORM_EXCEPTIONS_DISABLED
 		try
 #endif
@@ -747,7 +747,7 @@ void UOpenColorIOColorTransform::AllColorTransformsCacheResourceShadersForRender
 
 bool UOpenColorIOColorTransform::GenerateColorTransformData(const FString& InSourceColorSpace, const FString& InDestinationColorSpace)
 {
-#if WITH_EDITOR && WITH_OCIO
+#if WITH_OCIO
 	if (InSourceColorSpace.IsEmpty() || InDestinationColorSpace.IsEmpty())
 	{
 		return false;
@@ -767,7 +767,7 @@ bool UOpenColorIOColorTransform::GenerateColorTransformData(const FString& InSou
 
 bool UOpenColorIOColorTransform::GenerateColorTransformData(const FString& InSourceColorSpace, const FString& InDisplay, const FString& InView, EOpenColorIOViewTransformDirection InDirection)
 {
-#if WITH_EDITOR && WITH_OCIO
+#if WITH_OCIO
 	if (InSourceColorSpace.IsEmpty() || InDisplay.IsEmpty() || InView.IsEmpty())
 	{
 		return false;
@@ -882,7 +882,7 @@ TObjectPtr<UTexture> UOpenColorIOColorTransform::CreateTexture3DLUT(const FStrin
 {
 	TObjectPtr<UVolumeTexture> OutTexture = nullptr;
 
-#if WITH_EDITOR && WITH_OCIO
+#if WITH_OCIO
 	check(InSourceData);
 	const UOpenColorIOSettings* Settings = GetDefault<UOpenColorIOSettings>();
 
@@ -945,7 +945,7 @@ TObjectPtr<UTexture> UOpenColorIOColorTransform::CreateTexture1DLUT(const FStrin
 {
 	TObjectPtr<UTexture2D> OutTexture = nullptr;
 
-#if WITH_EDITOR && WITH_OCIO
+#if WITH_OCIO
 	check(InSourceData);
 	const UOpenColorIOSettings* Settings = GetDefault<UOpenColorIOSettings>();
 
diff --git a/Engine/Plugins/Compositing/OpenColorIO/Source/OpenColorIO/Private/OpenColorIOConfiguration.cpp b/Engine/Plugins/Compositing/OpenColorIO/Source/OpenColorIO/Private/OpenColorIOConfiguration.cpp
--- a/Engine/Plugins/Compositing/OpenColorIO/Source/OpenColorIO/Private/OpenColorIOConfiguration.cpp
+++ b/Engine/Plugins/Compositing/OpenColorIO/Source/OpenColorIO/Private/OpenColorIOConfiguration.cpp
@@ -30,7 +30,7 @@
 #include "Interfaces/ITargetPlatform.h"
 #endif //WITH_EDITOR
 
-#if WITH_EDITOR && WITH_OCIO
+#if WITH_OCIO
 namespace OCIODirectoryWatcher
 {
 	/** OCIO supported extensions we should be checking for when something changes in the OCIO config folder. */
@@ -166,7 +166,7 @@ bool UOpenColorIOConfiguration::EditorTransformImage(const FOpenColorIOColorConv
 
 void UOpenColorIOConfiguration::ReloadExistingColorspaces()
 {
-#if WITH_EDITOR && WITH_OCIO
+#if WITH_OCIO
 	TArray<FOpenColorIOColorSpace> ColorSpacesToBeReloaded = DesiredColorSpaces;
 	TArray<FOpenColorIODisplayView> DisplayViewsToBeReloaded = DesiredDisplayViews;
 	DesiredColorSpaces.Reset();
@@ -242,7 +242,7 @@ void UOpenColorIOConfiguration::ReloadExistingColorspaces()
 
 void UOpenColorIOConfiguration::ConfigPathChangedEvent(const TArray<FFileChangeData>& InFileChanges, const FString InFileMountPath)
 {
-#if WITH_EDITOR && WITH_OCIO
+#if WITH_OCIO
 	// We want to stop reacting to these events while the message is still up.
 	if (WatchedDirectoryInfo.RawConfigChangedToast.IsValid())
 	{
@@ -296,7 +296,7 @@ FOpenColorIONativeConfiguration* UOpenColorIOConfiguration::GetNativeConfig_Inte
 
 void UOpenColorIOConfiguration::StartDirectoryWatch(const FString& FilePath)
 {
-#if WITH_EDITOR && WITH_OCIO
+#if WITH_OCIO
 	FDirectoryWatcherModule& DirectoryWatcherModule = FModuleManager::LoadModuleChecked<FDirectoryWatcherModule>(OCIODirectoryWatcher::NAME_DirectoryWatcher);
 	if (IDirectoryWatcher* DirectoryWatcher = DirectoryWatcherModule.Get())
 	{
@@ -319,7 +319,7 @@ void UOpenColorIOConfiguration::StartDirectoryWatch(const FString& FilePath)
 
 void UOpenColorIOConfiguration::StopDirectoryWatch()
 {
-#if WITH_EDITOR && WITH_OCIO
+#if WITH_OCIO
 	FDirectoryWatcherModule& DirectoryWatcherModule = FModuleManager::LoadModuleChecked<FDirectoryWatcherModule>(OCIODirectoryWatcher::NAME_DirectoryWatcher);
 	if (IDirectoryWatcher* DirectoryWatcher = DirectoryWatcherModule.Get())
 	{
@@ -595,7 +595,7 @@ void UOpenColorIOConfiguration::PostEditChangeProperty(FPropertyChangedEvent& Pr
 
 void UOpenColorIOConfiguration::LoadConfiguration()
 {
-#if WITH_EDITOR && WITH_OCIO
+#if WITH_OCIO
 	NativeConfig->Set(nullptr);
 
 	if (!ConfigurationFile.FilePath.IsEmpty())
diff --git a/Engine/Plugins/Compositing/OpenColorIO/Source/OpenColorIO/Private/OpenColorIOLibHandler.cpp b/Engine/Plugins/Compositing/OpenColorIO/Source/OpenColorIO/Private/OpenColorIOLibHandler.cpp
--- a/Engine/Plugins/Compositing/OpenColorIO/Source/OpenColorIO/Private/OpenColorIOLibHandler.cpp
+++ b/Engine/Plugins/Compositing/OpenColorIO/Source/OpenColorIO/Private/OpenColorIOLibHandler.cpp
@@ -18,7 +18,7 @@ void* FOpenColorIOLibHandler::LibHandle = nullptr;
 //--------------------------------------------------------------------
 bool FOpenColorIOLibHandler::Initialize()
 {
-#if WITH_EDITOR && WITH_OCIO && defined(OCIO_DLL_NAME)
+#if WITH_OCIO && defined(OCIO_DLL_NAME)
 	check(LibHandle == nullptr);
 
 	const FString PluginDir = IPluginManager::Get().FindPlugin(TEXT("OpenColorIO"))->GetBaseDir();
@@ -55,7 +55,7 @@ bool FOpenColorIOLibHandler::IsInitialized()
 
 void FOpenColorIOLibHandler::Shutdown()
 {
-#if WITH_EDITOR && WITH_OCIO && defined(OCIO_DLL_NAME)
+#if WITH_OCIO && defined(OCIO_DLL_NAME)
 	if (LibHandle != nullptr)
 	{
 		FPlatformProcess::FreeDllHandle(LibHandle);
diff --git a/Engine/Plugins/Compositing/OpenColorIO/Source/OpenColorIO/Private/OpenColorIONativeConfiguration.cpp b/Engine/Plugins/Compositing/OpenColorIO/Source/OpenColorIO/Private/OpenColorIONativeConfiguration.cpp
--- a/Engine/Plugins/Compositing/OpenColorIO/Source/OpenColorIO/Private/OpenColorIONativeConfiguration.cpp
+++ b/Engine/Plugins/Compositing/OpenColorIO/Source/OpenColorIO/Private/OpenColorIONativeConfiguration.cpp
@@ -6,7 +6,7 @@
 #include "OpenColorIOConfiguration.h"
 
 
-#if WITH_EDITOR && WITH_OCIO
+#if WITH_OCIO
 OCIO_NAMESPACE::ConstConfigRcPtr FOpenColorIONativeConfiguration::Get() const
 {
 	return Config;
diff --git a/Engine/Plugins/Compositing/OpenColorIO/Source/OpenColorIO/Private/OpenColorIONativeConfiguration.h b/Engine/Plugins/Compositing/OpenColorIO/Source/OpenColorIO/Private/OpenColorIONativeConfiguration.h
--- a/Engine/Plugins/Compositing/OpenColorIO/Source/OpenColorIO/Private/OpenColorIONativeConfiguration.h
+++ b/Engine/Plugins/Compositing/OpenColorIO/Source/OpenColorIO/Private/OpenColorIONativeConfiguration.h
@@ -4,7 +4,7 @@
 
 #include "HAL/Platform.h"
 
-#if WITH_EDITOR && WITH_OCIO
+#if WITH_OCIO
 #include "OpenColorIO/OpenColorIO.h"
 #endif
 
@@ -13,7 +13,7 @@ static constexpr ANSICHAR OpenColorIOInterchangeName[] = "aces_interchange";
 
 class FOpenColorIONativeConfiguration {
 public:
-#if WITH_EDITOR && WITH_OCIO
+#if WITH_OCIO
 	/* Native config object getter. */
 	OCIO_NAMESPACE::ConstConfigRcPtr Get() const;
 
@@ -23,13 +23,13 @@ public:
 private:
 	/* Loaded native config object. */
 	OCIO_NAMESPACE::ConstConfigRcPtr Config;
-#endif //WITH_EDITORONLY_DATA
+#endif
 };
 
 
 class FOpenColorIONativeInterchangeConfiguration {
 public:
-#if WITH_EDITOR && WITH_OCIO
+#if WITH_OCIO
 	FOpenColorIONativeInterchangeConfiguration();
 
 	/* Native config object getter. */
