OCIO: Facilitate the obtention of transform render pass resources (unifying all current uses cases).

#rnx

diff --git a/Engine/Plugins/Compositing/OpenColorIO/Source/OpenColorIO/Private/OpenColorIODisplayExtension.cpp b/Engine/Plugins/Compositing/OpenColorIO/Source/OpenColorIO/Private/OpenColorIODisplayExtension.cpp
--- a/Engine/Plugins/Compositing/OpenColorIO/Source/OpenColorIO/Private/OpenColorIODisplayExtension.cpp
+++ b/Engine/Plugins/Compositing/OpenColorIO/Source/OpenColorIO/Private/OpenColorIODisplayExtension.cpp
@@ -51,51 +51,24 @@ void FOpenColorIODisplayExtension::AddReferencedObjects(FReferenceCollector& Col
 
 void FOpenColorIODisplayExtension::SetupView(FSceneViewFamily& InViewFamily, FSceneView& InView)
 {
-	//Cache render resource so they are available on the render thread (Can't access UObjects on RT)
-	//If something fails, cache invalid resources to invalidate them
-	FOpenColorIOTransformResource* ShaderResource = nullptr;
-	TSortedMap<int32, FTextureResource*> TransformTextureResources;
-	const FOpenColorIOColorConversionSettings& ConversionSettings = DisplayConfiguration.ColorConfiguration;
+	FOpenColorIORenderPassResources PassResources = FOpenColorIORendering::GetRenderPassResources(DisplayConfiguration.ColorConfiguration, InViewFamily.GetFeatureLevel());
 
-	if (ConversionSettings.ConfigurationSource != nullptr)
+	if (PassResources.IsValid())
 	{
-		const bool bFoundTransform = ConversionSettings.ConfigurationSource->GetRenderResources(
-			InViewFamily.GetFeatureLevel()
-			, ConversionSettings
-			, ShaderResource
-			, TransformTextureResources);
+		//Force ToneCurve to be off while we'are alive to make sure the input color space is the working space : srgb linear
+		InViewFamily.EngineShowFlags.SetToneCurve(false);
+		// This flags sets tonampper to output to ETonemapperOutputDevice::LinearNoToneCurve
+		InViewFamily.SceneCaptureSource = SCS_FinalColorHDR;
 
-		if (bFoundTransform)
-		{
-			// Transform was found, so shader must be there but doesn't mean the actual shader is available
-			check(ShaderResource);
-			if (ShaderResource->GetShaderGameThread<FOpenColorIOPixelShader>().IsNull())
-			{
-				ensureMsgf(false, TEXT("Can't apply display look - Shader was invalid for Resource %s"), *ShaderResource->GetFriendlyName());
-
-				//Invalidate shader resource
-				ShaderResource = nullptr;
-			}
-			else
-			{
-				//Force ToneCurve to be off while we'are alive to make sure the input color space is the working space : srgb linear
-				InViewFamily.EngineShowFlags.SetToneCurve(false);
-				// This flags sets tonampper to output to ETonemapperOutputDevice::LinearNoToneCurve
-				InViewFamily.SceneCaptureSource = SCS_FinalColorHDR;
-
-				InView.FinalPostProcessSettings.bOverride_ToneCurveAmount = 1;
-				InView.FinalPostProcessSettings.ToneCurveAmount = 0.0;
-			}
-		}
+		InView.FinalPostProcessSettings.bOverride_ToneCurveAmount = 1;
+		InView.FinalPostProcessSettings.ToneCurveAmount = 0.0;
 	}
 
 	ENQUEUE_RENDER_COMMAND(ProcessColorSpaceTransform)(
-		[this, ShaderResource, TextureResources = MoveTemp(TransformTextureResources), TransformName = ConversionSettings.ToString()](FRHICommandListImmediate& RHICmdList)
+		[this, ResourcesRenderThread = MoveTemp(PassResources)](FRHICommandListImmediate& RHICmdList)
 		{
 			//Caches render thread resource to be used when applying configuration in PostRenderViewFamily_RenderThread
-			CachedResourcesRenderThread.ShaderResource = ShaderResource;
-			CachedResourcesRenderThread.TextureResources = TextureResources;
-			CachedResourcesRenderThread.TransformName = TransformName;
+			CachedResourcesRenderThread = ResourcesRenderThread;
 		}
 	);
 }
diff --git a/Engine/Plugins/Compositing/OpenColorIO/Source/OpenColorIO/Private/OpenColorIORendering.cpp b/Engine/Plugins/Compositing/OpenColorIO/Source/OpenColorIO/Private/OpenColorIORendering.cpp
--- a/Engine/Plugins/Compositing/OpenColorIO/Source/OpenColorIO/Private/OpenColorIORendering.cpp
+++ b/Engine/Plugins/Compositing/OpenColorIO/Source/OpenColorIO/Private/OpenColorIORendering.cpp
@@ -103,6 +103,37 @@ void FOpenColorIORendering::AddPass_RenderThread(
 	}
 }
 
+// static
+FOpenColorIORenderPassResources FOpenColorIORendering::GetRenderPassResources(const FOpenColorIOColorConversionSettings& InSettings, ERHIFeatureLevel::Type InFeatureLevel)
+{
+	FOpenColorIORenderPassResources Result = { nullptr, {}, InSettings.ToString() };
+
+	if (InSettings.ConfigurationSource != nullptr)
+	{
+		const bool bFoundTransform = InSettings.ConfigurationSource->GetRenderResources(
+			InFeatureLevel,
+			InSettings,
+			Result.ShaderResource,
+			Result.TextureResources
+		);
+
+		if (bFoundTransform)
+		{
+			// Transform was found, so shader must be there but doesn't mean the actual shader is available
+			check(Result.ShaderResource);
+			if (Result.ShaderResource->GetShaderGameThread<FOpenColorIOPixelShader>().IsNull())
+			{
+				ensureMsgf(false, TEXT("Can't apply display look - Shader was invalid for Resource %s"), *Result.ShaderResource->GetFriendlyName());
+
+				//Invalidate shader resource
+				Result.ShaderResource = nullptr;
+			}
+		}
+	}
+
+	return Result;
+}
+
 // static
 bool FOpenColorIORendering::ApplyColorTransform(UWorld* InWorld, const FOpenColorIOColorConversionSettings& InSettings, UTexture* InTexture, UTextureRenderTarget2D* OutRenderTarget)
 {
diff --git a/Engine/Plugins/Compositing/OpenColorIO/Source/OpenColorIO/Public/OpenColorIORendering.h b/Engine/Plugins/Compositing/OpenColorIO/Source/OpenColorIO/Public/OpenColorIORendering.h
--- a/Engine/Plugins/Compositing/OpenColorIO/Source/OpenColorIO/Public/OpenColorIORendering.h
+++ b/Engine/Plugins/Compositing/OpenColorIO/Source/OpenColorIO/Public/OpenColorIORendering.h
@@ -28,6 +28,9 @@ struct OPENCOLORIO_API FOpenColorIORenderPassResources
 
 	/** Color transform string description. */
 	FString TransformName = FString();
+
+	/** Resources are valid when the shader resource is not null. */
+	bool IsValid() const { return ShaderResource != nullptr; }
 };
 
 /** Option to transform the render target alpha before applying the color transform. */
@@ -55,6 +58,16 @@ public:
 	 */
 	static bool ApplyColorTransform(UWorld* InWorld, const FOpenColorIOColorConversionSettings& InSettings, UTexture* InTexture, UTextureRenderTarget2D* OutRenderTarget);
 
+
+	/**
+	 * Get render pass resources on the game thread, to be enqueue-copied for later use by the `AddPass_RenderThread` function.
+	 *
+	 * @param InSettings Settings describing the color space transform to apply
+	 * @param InFeatureLevel Rendering feature level
+	 * @return FOpenColorIORenderPassResources Render pass resources to be enqueue-copied onto the render thread.
+	 */
+	static FOpenColorIORenderPassResources GetRenderPassResources(const FOpenColorIOColorConversionSettings& InSettings, ERHIFeatureLevel::Type InFeatureLevel);
+
 	/**
 	 * Applies the color transform RDG pass with the provided resources.
 	 *
diff --git a/Engine/Plugins/Media/MediaIOFramework/Source/MediaIOCore/Private/MediaCaptureRenderPass.cpp b/Engine/Plugins/Media/MediaIOFramework/Source/MediaIOCore/Private/MediaCaptureRenderPass.cpp
--- a/Engine/Plugins/Media/MediaIOFramework/Source/MediaIOCore/Private/MediaCaptureRenderPass.cpp
+++ b/Engine/Plugins/Media/MediaIOFramework/Source/MediaIOCore/Private/MediaCaptureRenderPass.cpp
@@ -148,34 +148,7 @@ namespace UE::MediaCapture::ColorConversion
 
 		if (MediaCapture && MediaCapture->GetDesiredCaptureOptions().ColorConversionSettings.IsValid())
 		{
-			const FMediaCaptureOptions& DesiredCaptureOptions = MediaCapture->GetDesiredCaptureOptions();
-			FOpenColorIOTransformResource* ShaderResource = nullptr;
-			TSortedMap<int32, FTextureResource*> TransformTextureResources;
-
-			if (DesiredCaptureOptions.ColorConversionSettings.ConfigurationSource != nullptr)
-			{
-				const bool bFoundTransform = DesiredCaptureOptions.ColorConversionSettings.ConfigurationSource->GetRenderResources(
-					GMaxRHIFeatureLevel
-					, DesiredCaptureOptions.ColorConversionSettings
-					, ShaderResource
-					, TransformTextureResources);
-
-				if (bFoundTransform)
-				{
-					// Transform was found, so shader must be there but doesn't mean the actual shader is available
-					check(ShaderResource);
-					if (ShaderResource->GetShaderGameThread<FOpenColorIOPixelShader>().IsNull())
-					{
-						ensureMsgf(false, TEXT("Can't apply display look - Shader was invalid for Resource %s"), *ShaderResource->GetFriendlyName());
-
-						//Invalidate shader resource
-						ShaderResource = nullptr;
-					}
-				}
-			}
-
-			OCIOResources.ShaderResource = ShaderResource;
-			OCIOResources.TextureResources = TransformTextureResources;
+			OCIOResources = FOpenColorIORendering::GetRenderPassResources(MediaCapture->GetDesiredCaptureOptions().ColorConversionSettings, GMaxRHIFeatureLevel);
 		}
 
 		return OCIOResources;
diff --git a/Engine/Plugins/Runtime/nDisplay/Source/DisplayCluster/Private/Render/Viewport/DisplayClusterViewport_OpenColorIO.cpp b/Engine/Plugins/Runtime/nDisplay/Source/DisplayCluster/Private/Render/Viewport/DisplayClusterViewport_OpenColorIO.cpp
--- a/Engine/Plugins/Runtime/nDisplay/Source/DisplayCluster/Private/Render/Viewport/DisplayClusterViewport_OpenColorIO.cpp
+++ b/Engine/Plugins/Runtime/nDisplay/Source/DisplayCluster/Private/Render/Viewport/DisplayClusterViewport_OpenColorIO.cpp
@@ -56,50 +56,15 @@ void FDisplayClusterViewport_OpenColorIO::UpdateOpenColorIORenderPassResources()
 		return;
 	}
 
-	FOpenColorIOTransformResource* ShaderResource = nullptr;
-	TSortedMap<int32, FTextureResource*> TransformTextureResources;
+	FOpenColorIORenderPassResources PassResources = FOpenColorIORendering::GetRenderPassResources(ConversionSettings, GMaxRHIFeatureLevel);
 
-	if (ConversionSettings.ConfigurationSource != nullptr)
-	{
-		const ERHIFeatureLevel::Type FeatureLevel = GMaxRHIFeatureLevel;
-
-		const bool bFoundTransform = ConversionSettings.ConfigurationSource->GetRenderResources(
-			FeatureLevel
-			, ConversionSettings
-			, ShaderResource
-			, TransformTextureResources);
-
-		if (bFoundTransform)
-		{
-			// Transform was found, so shader must be there but doesn't mean the actual shader is available
-			check(ShaderResource);
-			if (ShaderResource->GetShaderGameThread<FOpenColorIOPixelShader>().IsNull())
-			{
-				ensureMsgf(false, TEXT("Can't apply display look - Shader was invalid for Resource %s"), *ShaderResource->GetFriendlyName());
-
-				//Invalidate shader resource
-				ShaderResource = nullptr;
-			}
-			else
-			{
-				bShaderResourceValid = true;
-			}
-		}
-	}
-
-	if (!bShaderResourceValid)
-	{
-		ShaderResource = nullptr;
-		TransformTextureResources.Empty();
-	}
+	bShaderResourceValid = PassResources.IsValid();
 
 	ENQUEUE_RENDER_COMMAND(ProcessColorSpaceTransform)(
-		[This = SharedThis(this), ShaderResource, TextureResources = MoveTemp(TransformTextureResources), TransformName = ConversionSettings.ToString()](FRHICommandListImmediate& RHICmdList)
+		[This = SharedThis(this), PassResourcesRenderThread = MoveTemp(PassResources)](FRHICommandListImmediate& RHICmdList)
 		{
 			//Caches render thread resource to be used when applying configuration in PostRenderViewFamily_RenderThread
-			This->CachedResourcesRenderThread.ShaderResource = ShaderResource;
-			This->CachedResourcesRenderThread.TextureResources = TextureResources;
-			This->CachedResourcesRenderThread.TransformName = TransformName;
+			This->CachedResourcesRenderThread = PassResourcesRenderThread;
 		}
 	);
 }
