OCIO: Plugin refactor in preparation of engine integration, all native OCIO operations have been refactored into a new module.

In a secondary CL, OpenColorIOWrapper would be moved to Engine/Source/Developer.

diff --git a/Engine/Plugins/Compositing/OpenColorIO/OpenColorIO.uplugin b/Engine/Plugins/Compositing/OpenColorIO/OpenColorIO.uplugin
--- a/Engine/Plugins/Compositing/OpenColorIO/OpenColorIO.uplugin
+++ b/Engine/Plugins/Compositing/OpenColorIO/OpenColorIO.uplugin
@@ -15,6 +15,11 @@
 	"IsBetaVersion" : true,
 	"Installed" : false,
 	"Modules": [
+		{
+			"Name": "OpenColorIOWrapper",
+			"Type": "Editor",
+			"LoadingPhase": "PostConfigInit"
+		},
 		{
 			"Name": "OpenColorIO",
 			"Type": "Runtime",
diff --git a/Engine/Plugins/Compositing/OpenColorIO/Source/OpenColorIO/OpenColorIO.Build.cs b/Engine/Plugins/Compositing/OpenColorIO/Source/OpenColorIO/OpenColorIO.Build.cs
--- a/Engine/Plugins/Compositing/OpenColorIO/Source/OpenColorIO/OpenColorIO.Build.cs
+++ b/Engine/Plugins/Compositing/OpenColorIO/Source/OpenColorIO/OpenColorIO.Build.cs
@@ -51,15 +51,12 @@ namespace UnrealBuildTool.Rules
 				PrivateDependencyModuleNames.AddRange(
 					new string[]
 					{
-						"OpenColorIOLib",
 						"TargetPlatform",
 						"EditorFramework",
-						"UnrealEd"
+						"UnrealEd",
 					});
-			}
-			else
-			{
-				PrivateDefinitions.Add("WITH_OCIO=0");
+
+				PublicDependencyModuleNames.Add("OpenColorIOWrapper");
 			}
 		}
 	}
diff --git a/Engine/Plugins/Compositing/OpenColorIO/Source/OpenColorIO/Private/OpenColorIOColorTransform.cpp b/Engine/Plugins/Compositing/OpenColorIO/Source/OpenColorIO/Private/OpenColorIOColorTransform.cpp
--- a/Engine/Plugins/Compositing/OpenColorIO/Source/OpenColorIO/Private/OpenColorIOColorTransform.cpp
+++ b/Engine/Plugins/Compositing/OpenColorIO/Source/OpenColorIO/Private/OpenColorIOColorTransform.cpp
@@ -8,7 +8,6 @@
 #include "Modules/ModuleManager.h"
 #include "OpenColorIOConfiguration.h"
 #include "OpenColorIOModule.h"
-#include "OpenColorIONativeConfiguration.h"
 #include "OpenColorIOSettings.h"
 #include "UObject/UObjectIterator.h"
 #include "DataDrivenShaderPlatformInfo.h"
@@ -20,114 +19,42 @@
 #include "DerivedDataCacheInterface.h"
 #include "Editor.h"
 #include "Interfaces/ITargetPlatform.h"
+#include "OpenColorIOWrapper.h"
 #include "OpenColorIODerivedDataVersion.h"
 #include "OpenColorIOShader.h"
 
-#if WITH_OCIO
-#include "OpenColorIO/OpenColorIO.h"
-#include <vector>
-#endif
-
-#endif //WITH_EDITOR
-
-namespace {
-#if WITH_OCIO
-	/*
-	 * Get the processor optimization flag.
-	 */
-	OCIO_NAMESPACE::OptimizationFlags GetProcessorOptimization()
-	{
-		using namespace OCIO_NAMESPACE;
-
-		return static_cast<OptimizationFlags>(OptimizationFlags::OPTIMIZATION_DEFAULT | OptimizationFlags::OPTIMIZATION_NO_DYNAMIC_PROPERTIES);
-	}
-
-	/*
-	 * Get the OpenColorIO processor using the transform's config.
-	 * NOTE: By default, the library automatically caches and reuses processors.
-	 */
-	OCIO_NAMESPACE::ConstProcessorRcPtr GetTransformProcessor(const UOpenColorIOColorTransform* InTransform)
+namespace
+{
+	// Note: Since OpenColorIO caches processors automatically, we can recreate them without significant additional costs.
+	FOpenColorIOProcessorWrapper GetTransformProcessor(const UOpenColorIOColorTransform* InTransform, EOpenColorIOWorkingColorSpaceTransform InWCSTrasnformType)
 	{
 		check(InTransform);
 		check(InTransform->ConfigurationOwner);
 
-		OCIO_NAMESPACE::ConstConfigRcPtr Config = InTransform->ConfigurationOwner->GetNativeConfig_Internal()->Get();
-		if (Config)
+		EOpenColorIOViewTransformDirection DisplayViewDirection;
+		if (InTransform->GetDisplayViewDirection(DisplayViewDirection))
 		{
-			OCIO_NAMESPACE::ContextRcPtr Context = Config->getCurrentContext()->createEditableCopy();
-
-			for (const TPair<FString, FString>& KeyValue : InTransform->GetContextKeyValues())
-			{
-				Context->setStringVar(TCHAR_TO_ANSI(*KeyValue.Key), TCHAR_TO_ANSI(*KeyValue.Value));
-			}
-
-			EOpenColorIOViewTransformDirection DisplayViewDirection;
-			if (InTransform->GetDisplayViewDirection(DisplayViewDirection))
-			{
-				OCIO_NAMESPACE::TransformDirection OcioDirection = static_cast<OCIO_NAMESPACE::TransformDirection>(DisplayViewDirection);
-
-				return Config->getProcessor(
-					Context,
-					StringCast<ANSICHAR>(*InTransform->SourceColorSpace).Get(),
-					StringCast<ANSICHAR>(*InTransform->Display).Get(),
-					StringCast<ANSICHAR>(*InTransform->View).Get(),
-					OcioDirection);
-			}
-			else
-			{
-				return Config->getProcessor(
-					Context,
-					StringCast<ANSICHAR>(*InTransform->SourceColorSpace).Get(),
-					StringCast<ANSICHAR>(*InTransform->DestinationColorSpace).Get()
-				);
-			}
+			return FOpenColorIOProcessorWrapper(
+				InTransform->ConfigurationOwner->GetConfigWrapper(),
+				InTransform->SourceColorSpace,
+				InTransform->Display,
+				InTransform->View,
+				DisplayViewDirection == EOpenColorIOViewTransformDirection::Inverse,
+				InWCSTrasnformType,
+				InTransform->GetContextKeyValues());
 		}
-
-		return nullptr;
-	}
-
-	TUniquePtr<OCIO_NAMESPACE::PackedImageDesc> GetImageDesc(const FImageView& InImage)
-	{
-		OCIO_NAMESPACE::ChannelOrdering Ordering;
-		OCIO_NAMESPACE::BitDepth BitDepth;
-
-		switch (InImage.Format)
+		else
 		{
-		case ERawImageFormat::BGRA8:
-			Ordering = OCIO_NAMESPACE::CHANNEL_ORDERING_BGRA;
-			BitDepth = OCIO_NAMESPACE::BIT_DEPTH_UINT8;
-			break;
-		case ERawImageFormat::RGBA16:
-			Ordering = OCIO_NAMESPACE::CHANNEL_ORDERING_RGBA;
-			BitDepth = OCIO_NAMESPACE::BIT_DEPTH_UINT16;
-			break;
-		case ERawImageFormat::RGBA16F:
-			Ordering = OCIO_NAMESPACE::CHANNEL_ORDERING_RGBA;
-			BitDepth = OCIO_NAMESPACE::BIT_DEPTH_F16;
-			break;
-		case ERawImageFormat::RGBA32F:
-			Ordering = OCIO_NAMESPACE::CHANNEL_ORDERING_RGBA;
-			BitDepth = OCIO_NAMESPACE::BIT_DEPTH_F32;
-			break;
-		default:
-
-			UE_LOG(LogOpenColorIO, Log, TEXT("Unsupported texture format."));
-			return nullptr;
+			return FOpenColorIOProcessorWrapper(
+				InTransform->ConfigurationOwner->GetConfigWrapper(),
+				InTransform->SourceColorSpace,
+				InTransform->DestinationColorSpace,
+				InWCSTrasnformType,
+				InTransform->GetContextKeyValues());
 		}
-
-		return MakeUnique<OCIO_NAMESPACE::PackedImageDesc>(
-			InImage.RawData,
-			static_cast<long>(InImage.GetWidth()),
-			static_cast<long>(InImage.GetHeight()),
-			Ordering,
-			BitDepth,
-			OCIO_NAMESPACE::AutoStride,
-			OCIO_NAMESPACE::AutoStride,
-			OCIO_NAMESPACE::AutoStride
-		);
-	};
-#endif
+	}
 }
+#endif //WITH_EDITOR
 
 
 void UOpenColorIOColorTransform::SerializeOpenColorIOShaderMaps(const TMap<const ITargetPlatform*, TArray<FOpenColorIOTransformResource*>>* PlatformColorTransformResourcesToSavePtr, FArchive& Ar, TArray<FOpenColorIOTransformResource>&  OutLoadedResources)
@@ -203,11 +130,9 @@ void UOpenColorIOColorTransform::GetOpenColorIOLUTKeyGuid(const FString& InProce
 #if WITH_EDITOR
 	FString DDCKey = FDerivedDataCacheInterface::BuildCacheKey(TEXT("OCIOLUT"), OPENCOLORIO_DERIVEDDATA_VER, *InProcessorIdentifier);
 
-#if WITH_OCIO
 	//Keep library version in the DDC key to invalidate it once we move to a new library
 	DDCKey += TEXT("OCIOVersion");
-	DDCKey += TEXT(OCIO_VERSION);
-#endif //WITH_OCIO
+	DDCKey += OpenColorIOWrapper::GetVersion();
 
 	if (!InName.IsNone())
 	{
@@ -232,14 +157,16 @@ bool UOpenColorIOColorTransform::Initialize(UOpenColorIOConfiguration* InOwner,
 	ConfigurationOwner = InOwner;
 	ContextKeyValues = InContextKeyValues;
 	
-	if (InSourceColorSpace == UOpenColorIOConfiguration::WorkingColorSpaceName)
+#if WITH_EDITOR
+	if (InSourceColorSpace == OpenColorIOWrapper::WorkingColorSpaceName)
 	{
 		WorkingColorSpaceTransformType = EOpenColorIOWorkingColorSpaceTransform::Source;
 	}
-	else if(InDestinationColorSpace == UOpenColorIOConfiguration::WorkingColorSpaceName)
+	else if(InDestinationColorSpace == OpenColorIOWrapper::WorkingColorSpaceName)
 	{
 		WorkingColorSpaceTransformType = EOpenColorIOWorkingColorSpaceTransform::Destination;
 	}
+#endif
 	
 	return GenerateColorTransformData(InSourceColorSpace, InDestinationColorSpace);
 }
@@ -250,7 +177,8 @@ bool UOpenColorIOColorTransform::Initialize(UOpenColorIOConfiguration* InOwner,
 	ConfigurationOwner = InOwner;
 	ContextKeyValues = InContextKeyValues;
 
-	if (InSourceColorSpace == UOpenColorIOConfiguration::WorkingColorSpaceName)
+#if WITH_EDITOR
+	if (InSourceColorSpace == OpenColorIOWrapper::WorkingColorSpaceName)
 	{
 		if (InDirection == EOpenColorIOViewTransformDirection::Forward)
 		{
@@ -261,6 +189,7 @@ bool UOpenColorIOColorTransform::Initialize(UOpenColorIOConfiguration* InOwner,
 			WorkingColorSpaceTransformType = EOpenColorIOWorkingColorSpaceTransform::Destination;
 		}
 	}
+#endif
 
 	return GenerateColorTransformData(InSourceColorSpace, InDisplay, InView, InDirection);
 }
@@ -331,112 +260,62 @@ void UOpenColorIOColorTransform::SerializeLuts(FArchive& Ar)
 
 void UOpenColorIOColorTransform::CacheResourceTextures()
 {
+#if WITH_EDITOR
 	if (Textures.IsEmpty())
 	{
-#if WITH_OCIO
-#if !PLATFORM_EXCEPTIONS_DISABLED
-		try
-#endif
+		const FOpenColorIOProcessorWrapper TransformProcessor = GetTransformProcessor(this, WorkingColorSpaceTransformType);
+		if (TransformProcessor.IsValid())
 		{
-			OCIO_NAMESPACE::ConstProcessorRcPtr TransformProcessor = GetTransformProcessor(this);
-			if (TransformProcessor)
-			{
-				OCIO_NAMESPACE::GpuShaderDescRcPtr ShaderDescription = OCIO_NAMESPACE::GpuShaderDesc::CreateShaderDesc();
-				ShaderDescription->setLanguage(OCIO_NAMESPACE::GPU_LANGUAGE_HLSL_DX11);
-				ShaderDescription->setFunctionName(StringCast<ANSICHAR>(OpenColorIOShader::OpenColorIOShaderFunctionName).Get());
-				ShaderDescription->setResourcePrefix("Ocio");
+			const FOpenColorIOGPUProcessorWrapper GPUProcessor = FOpenColorIOGPUProcessorWrapper(TransformProcessor);
+			const FString ProcessorID = GPUProcessor.GetCacheID();
 
-				OCIO_NAMESPACE::ConstGPUProcessorRcPtr GPUProcessor = nullptr;
-				const UOpenColorIOSettings* Settings = GetDefault<UOpenColorIOSettings>();
-				check(Settings);
+			//In editor, it will use what's on DDC if there's something corresponding to the actual data or use that raw data
+			//that OCIO library has on board. The textures will be serialized only when cooking.
 
-				if (Settings->bUseLegacyProcessor)
-				{
-					unsigned int EdgeLength = static_cast<unsigned int>(OpenColorIOShader::Lut3dEdgeLength);
-					GPUProcessor = TransformProcessor->getOptimizedLegacyGPUProcessor(GetProcessorOptimization(), EdgeLength);
-				}
-				else
-				{
-					GPUProcessor = TransformProcessor->getOptimizedGPUProcessor(GetProcessorOptimization());
-				}
+			// Process 3D luts
+			for (uint32 Index = 0; Index < GPUProcessor.GetNum3DTextures(); ++Index)
+			{
+				TextureFilter Filter;
+				FName TextureName;
+				uint32 EdgeLength = 0;
+				const float* TextureValues = 0x0;
 
-				GPUProcessor->extractGpuShaderInfo(ShaderDescription);
-				
-				ensureMsgf(ShaderDescription->getNumDynamicProperties() == 0, TEXT("Dynamic properties are not currently supported."));
+				bool bSuccess = GPUProcessor.Get3DTexture(Index, TextureName, EdgeLength, Filter, TextureValues);
+				check(bSuccess);
 
-				const FString ProcessorID = StringCast<TCHAR>(GPUProcessor->getCacheID()).Get();
+				TObjectPtr<UTexture> Result = CreateTexture3DLUT(ProcessorID, TextureName, EdgeLength, Filter, TextureValues);
 
-				//In editor, it will use what's on DDC if there's something corresponding to the actual data or use that raw data
-				//that OCIO library has on board. The textures will be serialized only when cooking.
+				const int32 SlotIndex = TextureName.GetNumber() - 1; // Rely on FName's index number extraction for convenience
+				Textures.Add(SlotIndex, MoveTemp(Result));
+			}
 
-				// Process 3D luts
-				for (uint32 Index = 0; Index < ShaderDescription->getNum3DTextures(); ++Index)
-				{
-					const char* TextureName = nullptr;
-					const char* SamplerName = nullptr;
-					unsigned int EdgeLength;
-					OCIO_NAMESPACE::Interpolation Interpolation = OCIO_NAMESPACE::INTERP_TETRAHEDRAL;
-
-					ShaderDescription->get3DTexture(Index, TextureName, SamplerName, EdgeLength, Interpolation);
-					checkf(TextureName && *TextureName && SamplerName && *SamplerName && EdgeLength > 0, TEXT("Invalid OCIO 3D texture or sampler."));
-
-					const float* TextureValues = 0x0;
-					ShaderDescription->get3DTextureValues(Index, TextureValues);
-					checkf(TextureValues, TEXT("Failed to read OCIO 3D LUT data."));
-
-					TextureFilter Filter = TF_Bilinear;
-					if (Interpolation == OCIO_NAMESPACE::Interpolation::INTERP_NEAREST || Interpolation == OCIO_NAMESPACE::Interpolation::INTERP_TETRAHEDRAL)
-					{
-						Filter = TF_Nearest;
-					}
-					const FName TextureFName = FName(TextureName);
-					TObjectPtr<UTexture> Result = CreateTexture3DLUT(ProcessorID, TextureFName, EdgeLength, Filter, TextureValues);
-
-					const int32 SlotIndex = TextureFName.GetNumber() - 1; // Rely on FName's index number extraction for convenience
-					Textures.Add(SlotIndex, MoveTemp(Result));
-				}
+			// Process 1D luts
+			for (uint32 Index = 0; Index < GPUProcessor.GetNumTextures(); ++Index)
+			{
+				FName TextureName;
+				uint32 TextureWidth = 0;
+				uint32 TextureHeight = 0;
+				TextureFilter Filter;
+				bool bRedChannelOnly = false;
+				const float* TextureValues = 0x0;
 
-				// Process 1D luts
-				for (uint32 Index = 0; Index < ShaderDescription->getNumTextures(); ++Index)
-				{
-					const char* TextureName = nullptr;
-					const char* SamplerName = nullptr;
-					unsigned TextureWidth = 0;
-					unsigned TextureHeight = 0;
-					OCIO_NAMESPACE::GpuShaderDesc::TextureType Channel = OCIO_NAMESPACE::GpuShaderDesc::TEXTURE_RGB_CHANNEL;
-					OCIO_NAMESPACE::Interpolation Interpolation = OCIO_NAMESPACE::Interpolation::INTERP_LINEAR;
-
-					ShaderDescription->getTexture(Index, TextureName, SamplerName, TextureWidth, TextureHeight, Channel, Interpolation);
-					checkf(TextureName && *TextureName && SamplerName && *SamplerName && TextureWidth > 0, TEXT("Invalid OCIO 1D texture or sampler."));
-
-					const float* TextureValues = 0x0;
-					ShaderDescription->getTextureValues(Index, TextureValues);
-					checkf(TextureValues, TEXT("Failed to read OCIO 1D LUT data."));
-
-					TextureFilter Filter = Interpolation == OCIO_NAMESPACE::Interpolation::INTERP_NEAREST ? TF_Nearest : TF_Bilinear;
-					bool bRedChannelOnly = Channel == OCIO_NAMESPACE::GpuShaderCreator::TEXTURE_RED_CHANNEL;
-					const FName TextureFName = FName(TextureName);
-					TObjectPtr<UTexture> Result = CreateTexture1DLUT(ProcessorID, TextureFName, TextureWidth, TextureHeight, Filter, bRedChannelOnly, TextureValues);
-
-					const int32 SlotIndex = TextureFName.GetNumber() - 1; // Rely on FName's index number extraction for convenience
-					Textures.Add(SlotIndex, MoveTemp(Result));
-				}
+				bool bSuccess = GPUProcessor.GetTexture(Index, TextureName, TextureWidth, TextureHeight, Filter, bRedChannelOnly, TextureValues);
+				checkf(bSuccess, TEXT("Failed to read OCIO 1D LUT data."));
 
-				ensureAlwaysMsgf(Textures.Num() <= (int32)OpenColorIOShader::MaximumTextureSlots, TEXT("Color transform %s exceeds our current limit of %u texture slots. Use the legacy processor instead."), *GetTransformFriendlyName(), OpenColorIOShader::MaximumTextureSlots);
-			}
-			else
-			{
-				UE_LOG(LogOpenColorIO, Error, TEXT("Failed to cache texture resource(s) for color transform %s. Configuration file [%s] was invalid."), *GetTransformFriendlyName(), *ConfigurationOwner->ConfigurationFile.FilePath);
+				TObjectPtr<UTexture> Result = CreateTexture1DLUT(ProcessorID, TextureName, TextureWidth, TextureHeight, Filter, bRedChannelOnly, TextureValues);
+
+				const int32 SlotIndex = TextureName.GetNumber() - 1; // Rely on FName's index number extraction for convenience
+				Textures.Add(SlotIndex, MoveTemp(Result));
 			}
+
+			ensureAlwaysMsgf(Textures.Num() <= (int32)OpenColorIOShader::MaximumTextureSlots, TEXT("Color transform %s exceeds our current limit of %u texture slots. Use the legacy processor instead."), *GetTransformFriendlyName(), OpenColorIOShader::MaximumTextureSlots);
 		}
-#if !PLATFORM_EXCEPTIONS_DISABLED
-		catch (OCIO_NAMESPACE::Exception& exception)
+		else
 		{
-			UE_LOG(LogOpenColorIO, Error, TEXT("Failed to cache texture resource(s) for color transform %s. Error message: %s"), *GetTransformFriendlyName(), StringCast<TCHAR>(exception.what()).Get());
+			UE_LOG(LogOpenColorIO, Error, TEXT("Failed to cache texture resource(s) for color transform %s. Configuration file [%s] was invalid."), *GetTransformFriendlyName(), *ConfigurationOwner->ConfigurationFile.FilePath);
 		}
-#endif
-#endif
 	}
+#endif // WITH_EDITOR
 }
 
 void UOpenColorIOColorTransform::CacheResourceShadersForRendering(bool bRegenerateId)
@@ -595,131 +474,26 @@ bool UOpenColorIOColorTransform::IsTransform(const FString& InSourceColorSpace,
 #if WITH_EDITOR
 bool UOpenColorIOColorTransform::EditorTransformImage(const FImageView& InOutImage) const
 {
-#if WITH_OCIO
-	using namespace OCIO_NAMESPACE;
-
-#if !PLATFORM_EXCEPTIONS_DISABLED
-	try
-#endif
+	const FOpenColorIOProcessorWrapper TransformProcessor = GetTransformProcessor(this, WorkingColorSpaceTransformType);
+	if (TransformProcessor.IsValid())
 	{
-		ConstProcessorRcPtr Processor = GetTransformProcessor(this);
-		if (Processor && ConfigurationOwner)
-		{
-			TUniquePtr<PackedImageDesc> ImageDesc = GetImageDesc(InOutImage);
-			if (ImageDesc)
-			{
-				ConstConfigRcPtr InterchangeConfig = IOpenColorIOModule::Get().GetNativeInterchangeConfig_Internal()->Get();
-				ConstConfigRcPtr Config = ConfigurationOwner->GetNativeConfig_Internal()->Get();
-				BitDepth BitDepth = ImageDesc->getBitDepth();
-
-				// Conditionally apply a conversion from the working color space to interchange space
-				if (WorkingColorSpaceTransformType == EOpenColorIOWorkingColorSpaceTransform::Source)
-				{
-					ConstProcessorRcPtr	InterchangeProcessor = InterchangeConfig->GetProcessorFromConfigs(
-						InterchangeConfig,
-						StringCast<ANSICHAR>(UOpenColorIOConfiguration::WorkingColorSpaceName).Get(),
-						Config,
-						Config->getCanonicalName(OpenColorIOInterchangeName));
-
-					ConstCPUProcessorRcPtr InterchangeCPUProcessor = InterchangeProcessor->getOptimizedCPUProcessor(BitDepth, BitDepth, OPTIMIZATION_DEFAULT);
-					InterchangeCPUProcessor->apply(*ImageDesc);
-				}
+		const FOpenColorIOCPUProcessorWrapper CPUProcessor = FOpenColorIOCPUProcessorWrapper(TransformProcessor);
 
-				// Apply the main color transformation
-				ConstCPUProcessorRcPtr CPUProcessor = Processor->getOptimizedCPUProcessor(BitDepth, BitDepth, OPTIMIZATION_DEFAULT);
-				CPUProcessor->apply(*ImageDesc);
-
-				// Conditionally apply a conversion from the interchange space to the working color space
-				if (WorkingColorSpaceTransformType == EOpenColorIOWorkingColorSpaceTransform::Destination)
-				{
-					ConstProcessorRcPtr	InterchangeProcessor = InterchangeConfig->GetProcessorFromConfigs(
-						Config,
-						Config->getCanonicalName(OpenColorIOInterchangeName),
-						InterchangeConfig,
-						StringCast<ANSICHAR>(UOpenColorIOConfiguration::WorkingColorSpaceName).Get());
-
-					ConstCPUProcessorRcPtr InterchangeCPUProcessor = InterchangeProcessor->getOptimizedCPUProcessor(BitDepth, BitDepth, OPTIMIZATION_DEFAULT);
-					InterchangeCPUProcessor->apply(*ImageDesc);
-				}
-				
-				return true;
-			}
-		}
-	}
-#if !PLATFORM_EXCEPTIONS_DISABLED
-	catch (OCIO_NAMESPACE::Exception& exception)
-	{
-		UE_LOG(LogOpenColorIO, Log, TEXT("Failed to transform image. Error message: %s"), StringCast<TCHAR>(exception.what()).Get());
+		return CPUProcessor.TransformImage(InOutImage);
 	}
-#endif
-#endif // WITH_OCIO
 
 	return false;
 }
 
 bool UOpenColorIOColorTransform::EditorTransformImage(const FImageView& SrcImage, const FImageView& DestImage) const
 {
-#if WITH_OCIO
-	using namespace OCIO_NAMESPACE;
-#if !PLATFORM_EXCEPTIONS_DISABLED
-	try
-#endif
+	const FOpenColorIOProcessorWrapper TransformProcessor = GetTransformProcessor(this, WorkingColorSpaceTransformType);
+	if (TransformProcessor.IsValid())
 	{
-		ConstProcessorRcPtr Processor = GetTransformProcessor(this);
-		if (Processor && ConfigurationOwner)
-		{
-			TUniquePtr<PackedImageDesc> SrcImageDesc = GetImageDesc(SrcImage);
-			TUniquePtr<PackedImageDesc> DestImageDesc = GetImageDesc(DestImage);
-			if (SrcImageDesc && DestImageDesc)
-			{
-				ConstConfigRcPtr	InterchangeConfig = IOpenColorIOModule::Get().GetNativeInterchangeConfig_Internal()->Get();
-				ConstConfigRcPtr	Config = ConfigurationOwner->GetNativeConfig_Internal()->Get();
-
-				BitDepth SrcBitDepth = SrcImageDesc->getBitDepth();
-				BitDepth DestBitDepth = DestImageDesc->getBitDepth();
-
-				// Conditionally apply a conversion from the working color space to interchange space
-				if (WorkingColorSpaceTransformType == EOpenColorIOWorkingColorSpaceTransform::Source)
-				{
-					ConstProcessorRcPtr	InterchangeProcessor = InterchangeConfig->GetProcessorFromConfigs(
-						InterchangeConfig,
-						StringCast<ANSICHAR>(UOpenColorIOConfiguration::WorkingColorSpaceName).Get(),
-						Config,
-						Config->getCanonicalName(OpenColorIOInterchangeName));
-
-					ConstCPUProcessorRcPtr InterchangeCPUProcessor = InterchangeProcessor->getOptimizedCPUProcessor(SrcBitDepth, SrcBitDepth, OPTIMIZATION_DEFAULT);
-					InterchangeCPUProcessor->apply(*SrcImageDesc);
-				}
-
-
-				// Apply the main color transformation
-				ConstCPUProcessorRcPtr CPUProcessor = Processor->getOptimizedCPUProcessor(SrcBitDepth, DestBitDepth, OPTIMIZATION_DEFAULT);
-				CPUProcessor->apply(*SrcImageDesc, *DestImageDesc);
+		const FOpenColorIOCPUProcessorWrapper CPUProcessor = FOpenColorIOCPUProcessorWrapper(TransformProcessor);
 
-				// Conditionally apply a conversion from the interchange space to the working color space
-				if (WorkingColorSpaceTransformType == EOpenColorIOWorkingColorSpaceTransform::Destination)
-				{
-					ConstProcessorRcPtr	InterchangeProcessor = InterchangeConfig->GetProcessorFromConfigs(
-						Config,
-						Config->getCanonicalName(OpenColorIOInterchangeName),
-						InterchangeConfig,
-						StringCast<ANSICHAR>(UOpenColorIOConfiguration::WorkingColorSpaceName).Get());
-
-					ConstCPUProcessorRcPtr InterchangeCPUProcessor = InterchangeProcessor->getOptimizedCPUProcessor(DestBitDepth, DestBitDepth, OPTIMIZATION_DEFAULT);
-					InterchangeCPUProcessor->apply(*DestImageDesc);
-				}
-
-				return true;
-			}
-		}
-	}
-#if !PLATFORM_EXCEPTIONS_DISABLED
-	catch (OCIO_NAMESPACE::Exception& exception)
-	{
-		UE_LOG(LogOpenColorIO, Log, TEXT("Failed to transform image. Error message: %s"), StringCast<TCHAR>(exception.what()).Get());
+		return CPUProcessor.TransformImage(SrcImage, DestImage);
 	}
-#endif
-#endif // WITH_OCIO
 
 	return false;
 }
@@ -747,7 +521,7 @@ void UOpenColorIOColorTransform::AllColorTransformsCacheResourceShadersForRender
 
 bool UOpenColorIOColorTransform::GenerateColorTransformData(const FString& InSourceColorSpace, const FString& InDestinationColorSpace)
 {
-#if WITH_OCIO
+#if WITH_EDITOR
 	if (InSourceColorSpace.IsEmpty() || InDestinationColorSpace.IsEmpty())
 	{
 		return false;
@@ -767,7 +541,7 @@ bool UOpenColorIOColorTransform::GenerateColorTransformData(const FString& InSou
 
 bool UOpenColorIOColorTransform::GenerateColorTransformData(const FString& InSourceColorSpace, const FString& InDisplay, const FString& InView, EOpenColorIOViewTransformDirection InDirection)
 {
-#if WITH_OCIO
+#if WITH_EDITOR
 	if (InSourceColorSpace.IsEmpty() || InDisplay.IsEmpty() || InView.IsEmpty())
 	{
 		return false;
@@ -812,77 +586,28 @@ FString UOpenColorIOColorTransform::GetTransformFriendlyName()
 bool UOpenColorIOColorTransform::UpdateShaderInfo(FString& OutShaderCodeHash, FString& OutShaderCode, FString& OutRawConfigHash)
 {
 #if WITH_EDITOR
-#if WITH_OCIO
-#if !PLATFORM_EXCEPTIONS_DISABLED
-	try
-#endif
-	{
-		OCIO_NAMESPACE::ConstProcessorRcPtr TransformProcessor = GetTransformProcessor(this);
-		if (TransformProcessor)
-		{
-			OCIO_NAMESPACE::GpuShaderDescRcPtr ShaderDescription = OCIO_NAMESPACE::GpuShaderDesc::CreateShaderDesc();
-			ShaderDescription->setLanguage(OCIO_NAMESPACE::GPU_LANGUAGE_HLSL_DX11);
-			ShaderDescription->setFunctionName(StringCast<ANSICHAR>(OpenColorIOShader::OpenColorIOShaderFunctionName).Get());
-			ShaderDescription->setResourcePrefix("Ocio");
+	const FOpenColorIOProcessorWrapper Processor = GetTransformProcessor(this, WorkingColorSpaceTransformType);
+	const FOpenColorIOGPUProcessorWrapper GPUProcessor = FOpenColorIOGPUProcessorWrapper(Processor);
 
-			OCIO_NAMESPACE::ConstGPUProcessorRcPtr GPUProcessor = nullptr;
-			const UOpenColorIOSettings* Settings = GetDefault<UOpenColorIOSettings>();
-			check(Settings);
-
-			if (Settings->bUseLegacyProcessor)
-			{
-				unsigned int EdgeLength = static_cast<unsigned int>(OpenColorIOShader::Lut3dEdgeLength);
-				GPUProcessor = TransformProcessor->getOptimizedLegacyGPUProcessor(GetProcessorOptimization(), EdgeLength);
-			}
-			else
-			{
-				GPUProcessor = TransformProcessor->getOptimizedGPUProcessor(GetProcessorOptimization());
-			}
-			GPUProcessor->extractGpuShaderInfo(ShaderDescription);
-
-			ensureMsgf(ShaderDescription->getNumDynamicProperties() == 0, TEXT("We do not currently support dynamic properties."));
-
-			FString GLSLShaderCode = StringCast<TCHAR>(ShaderDescription->getShaderText()).Get();
-
-			OutShaderCodeHash = StringCast<TCHAR>(ShaderDescription->getCacheID()).Get();
-			OutShaderCode = StringCast<TCHAR>(ShaderDescription->getShaderText()).Get();
-			OCIO_NAMESPACE::ConstConfigRcPtr CurrentConfig = ConfigurationOwner->GetNativeConfig_Internal()->Get();
-			OutRawConfigHash = StringCast<TCHAR>(CurrentConfig->getCacheID()).Get();
-			return true;
-		}
-		else
-		{
-			UE_LOG(LogOpenColorIO, Error, TEXT("Failed to fetch shader info for color transform %s. Configuration file [%s] was invalid."), *GetTransformFriendlyName(), *ConfigurationOwner->ConfigurationFile.FilePath);
-		}
-	}
-#if !PLATFORM_EXCEPTIONS_DISABLED
-	catch (OCIO_NAMESPACE::Exception& exception)
+	if (GPUProcessor.IsValid())
 	{
-		UE_LOG(LogOpenColorIO, Error, TEXT("Failed to fetch shader info for color transform %s. Error message: %s"), *GetTransformFriendlyName(), StringCast<TCHAR>(exception.what()).Get());
-	}
-#endif
+		OutRawConfigHash = ConfigurationOwner->GetConfigWrapper()->GetCacheID();
 
-	return false;
-#else
-	//Avoid triggering errors when building maps on build machine.
-#if PLATFORM_WINDOWS
-	if (!GIsBuildMachine)
-	{
-		UE_LOG(LogOpenColorIO, Error, TEXT("Can't update shader, OCIO library isn't present."));
+		return GPUProcessor.GetShader(OutShaderCode, OutShaderCodeHash);
 	}
-#endif //PLATFORM_WINDOWS
+	
+	UE_LOG(LogOpenColorIO, Error, TEXT("Failed to fetch shader info for color transform %s. Configuration file [%s] was invalid."), *GetTransformFriendlyName(), *ConfigurationOwner->ConfigurationFile.FilePath);
 	return false;
-#endif //WITH_OCIO
 #else
 	return true; //When not in editor, shaders have been cooked so we're not relying on the library data anymore.
-#endif
+#endif //WITH_EDITOR
 }
 
 TObjectPtr<UTexture> UOpenColorIOColorTransform::CreateTexture3DLUT(const FString& InProcessorIdentifier, const FName& InName, uint32 InLutLength, TextureFilter InFilter, const float* InSourceData)
 {
 	TObjectPtr<UVolumeTexture> OutTexture = nullptr;
 
-#if WITH_OCIO
+#if WITH_EDITOR
 	check(InSourceData);
 	const UOpenColorIOSettings* Settings = GetDefault<UOpenColorIOSettings>();
 
@@ -945,7 +670,7 @@ TObjectPtr<UTexture> UOpenColorIOColorTransform::CreateTexture1DLUT(const FStrin
 {
 	TObjectPtr<UTexture2D> OutTexture = nullptr;
 
-#if WITH_OCIO
+#if WITH_EDITOR
 	check(InSourceData);
 	const UOpenColorIOSettings* Settings = GetDefault<UOpenColorIOSettings>();
 
diff --git a/Engine/Plugins/Compositing/OpenColorIO/Source/OpenColorIO/Private/OpenColorIOConfiguration.cpp b/Engine/Plugins/Compositing/OpenColorIO/Source/OpenColorIO/Private/OpenColorIOConfiguration.cpp
--- a/Engine/Plugins/Compositing/OpenColorIO/Source/OpenColorIO/Private/OpenColorIOConfiguration.cpp
+++ b/Engine/Plugins/Compositing/OpenColorIO/Source/OpenColorIO/Private/OpenColorIOConfiguration.cpp
@@ -11,7 +11,6 @@
 #include "Modules/ModuleManager.h"
 #include "OpenColorIOColorTransform.h"
 #include "OpenColorIOModule.h"
-#include "OpenColorIONativeConfiguration.h"
 #include "OpenColorIOSettings.h"
 #include "TextureResource.h"
 #include "Widgets/Notifications/SNotificationList.h"
@@ -24,13 +23,12 @@
 
 
 #if WITH_EDITOR
+#include "OpenColorIOWrapper.h"
 #include "DerivedDataCacheInterface.h"
 #include "DirectoryWatcherModule.h"
 #include "IDirectoryWatcher.h"
 #include "Interfaces/ITargetPlatform.h"
-#endif //WITH_EDITOR
 
-#if WITH_OCIO
 namespace OCIODirectoryWatcher
 {
 	/** OCIO supported extensions we should be checking for when something changes in the OCIO config folder. */
@@ -50,7 +48,6 @@ namespace OCIODirectoryWatcher
 
 UOpenColorIOConfiguration::UOpenColorIOConfiguration(const FObjectInitializer& ObjectInitializer)
 	: Super(ObjectInitializer)
-	, NativeConfig(MakePimpl<FOpenColorIONativeConfiguration>())
 {
 	
 }
@@ -123,20 +120,14 @@ bool UOpenColorIOConfiguration::HasDesiredDisplayView(const FOpenColorIODisplayV
 
 bool UOpenColorIOConfiguration::Validate() const
 {
-#if WITH_EDITOR 
-
-#if WITH_OCIO
-	if (!ConfigurationFile.FilePath.IsEmpty())
+#if WITH_EDITOR
+	if (!ConfigurationFile.FilePath.IsEmpty() && Config.IsValid())
 	{
 		//When loading the configuration file, if any errors are detected, it will throw an exception. Thus, our pointer won't be valid.
-		return NativeConfig->Get() != nullptr;
+		return Config->IsValid();
 	}
 
 	return false;
-#else
-	return false;
-#endif // WITH_OCIO
-
 #else
 	return true;
 #endif // WITH_EDITOR
@@ -166,7 +157,7 @@ bool UOpenColorIOConfiguration::EditorTransformImage(const FOpenColorIOColorConv
 
 void UOpenColorIOConfiguration::ReloadExistingColorspaces()
 {
-#if WITH_OCIO
+#if WITH_EDITOR
 	TArray<FOpenColorIOColorSpace> ColorSpacesToBeReloaded = DesiredColorSpaces;
 	TArray<FOpenColorIODisplayView> DisplayViewsToBeReloaded = DesiredDisplayViews;
 	DesiredColorSpaces.Reset();
@@ -174,18 +165,16 @@ void UOpenColorIOConfiguration::ReloadExistingColorspaces()
 	CleanupTransforms();
 	LoadConfiguration();
 
-	OCIO_NAMESPACE::ConstConfigRcPtr LoadedConfig = NativeConfig->Get();
-	if (!LoadedConfig)
+	if (Config == nullptr)
 	{
 		return;
 	}
+
 	// This will make sure that all colorspaces are up to date in case an index, family or name is changed.
 	for (const FOpenColorIOColorSpace& ExistingColorSpace : ColorSpacesToBeReloaded)
 	{
-		const auto ColorSpaceName = StringCast<ANSICHAR>(*ExistingColorSpace.ColorSpaceName);
-		int ColorSpaceIndex = LoadedConfig->getIndexForColorSpace(ColorSpaceName.Get());
-		OCIO_NAMESPACE::ConstColorSpaceRcPtr LibColorSpace = LoadedConfig->getColorSpace(ColorSpaceName.Get());
-		if (!LibColorSpace)
+		int ColorSpaceIndex = Config->GetColorSpaceIndex(*ExistingColorSpace.ColorSpaceName);
+		if (ColorSpaceIndex < 0)
 		{
 			// Name not found, therefore we don't need to re-add this colorspace.
 			continue;
@@ -194,16 +183,14 @@ void UOpenColorIOConfiguration::ReloadExistingColorspaces()
 		FOpenColorIOColorSpace ColorSpace;
 		ColorSpace.ColorSpaceIndex = ColorSpaceIndex;
 		ColorSpace.ColorSpaceName = ExistingColorSpace.ColorSpaceName;
-		ColorSpace.FamilyName = StringCast<TCHAR>(LibColorSpace->getFamily()).Get();
+		ColorSpace.FamilyName = Config->GetColorSpaceFamilyName(*ExistingColorSpace.ColorSpaceName);
 		DesiredColorSpaces.Add(ColorSpace);
 	}
 
 	for (const FOpenColorIODisplayView& ExistingDisplayView : DisplayViewsToBeReloaded)
 	{
-		const auto DisplayName = StringCast<ANSICHAR>(*ExistingDisplayView.Display);
-		const auto ViewName = StringCast<ANSICHAR>(*ExistingDisplayView.View);
-		const auto TransformName = LoadedConfig->getDisplayViewTransformName(DisplayName.Get(), ViewName.Get());
-		if (TransformName == nullptr)
+		const FString TransformName = Config->GetDisplayViewTransformName(*ExistingDisplayView.Display, *ExistingDisplayView.View);
+		if (TransformName.IsEmpty())
 		{
 			// Name not found, therefore we don't need to re-add this display-view.
 			continue;
@@ -242,7 +229,7 @@ void UOpenColorIOConfiguration::ReloadExistingColorspaces()
 
 void UOpenColorIOConfiguration::ConfigPathChangedEvent(const TArray<FFileChangeData>& InFileChanges, const FString InFileMountPath)
 {
-#if WITH_OCIO
+#if WITH_EDITOR
 	// We want to stop reacting to these events while the message is still up.
 	if (WatchedDirectoryInfo.RawConfigChangedToast.IsValid())
 	{
@@ -289,14 +276,18 @@ void UOpenColorIOConfiguration::ConfigPathChangedEvent(const TArray<FFileChangeD
 #endif
 }
 
-FOpenColorIONativeConfiguration* UOpenColorIOConfiguration::GetNativeConfig_Internal() const
+FOpenColorIOConfigWrapper* UOpenColorIOConfiguration::GetConfigWrapper() const
 {
-	return NativeConfig.Get();
+#if WITH_EDITOR
+	return Config.Get();
+#else
+	return nullptr;
+#endif
 }
 
 void UOpenColorIOConfiguration::StartDirectoryWatch(const FString& FilePath)
 {
-#if WITH_OCIO
+#if WITH_EDITOR
 	FDirectoryWatcherModule& DirectoryWatcherModule = FModuleManager::LoadModuleChecked<FDirectoryWatcherModule>(OCIODirectoryWatcher::NAME_DirectoryWatcher);
 	if (IDirectoryWatcher* DirectoryWatcher = DirectoryWatcherModule.Get())
 	{
@@ -319,7 +310,7 @@ void UOpenColorIOConfiguration::StartDirectoryWatch(const FString& FilePath)
 
 void UOpenColorIOConfiguration::StopDirectoryWatch()
 {
-#if WITH_OCIO
+#if WITH_EDITOR
 	FDirectoryWatcherModule& DirectoryWatcherModule = FModuleManager::LoadModuleChecked<FDirectoryWatcherModule>(OCIODirectoryWatcher::NAME_DirectoryWatcher);
 	if (IDirectoryWatcher* DirectoryWatcher = DirectoryWatcherModule.Get())
 	{
@@ -595,59 +586,47 @@ void UOpenColorIOConfiguration::PostEditChangeProperty(FPropertyChangedEvent& Pr
 
 void UOpenColorIOConfiguration::LoadConfiguration()
 {
-#if WITH_OCIO
-	NativeConfig->Set(nullptr);
+#if WITH_EDITOR
+	Config.Reset();
 
-	if (!ConfigurationFile.FilePath.IsEmpty())
+	if (ConfigurationFile.FilePath.IsEmpty())
 	{
-#if !PLATFORM_EXCEPTIONS_DISABLED
-		try
-#endif
-		{
-			FString ConfigurationFilePath = ConfigurationFile.FilePath;
-			if (ConfigurationFilePath.StartsWith(TEXT("ocio://")))
-			{
-				NativeConfig->Set(OCIO_NAMESPACE::Config::CreateFromFile(StringCast<ANSICHAR>(*ConfigurationFilePath).Get()));
-				UE_LOG(LogOpenColorIO, Verbose, TEXT("Loaded built-in OCIO configuration file %s"), *ConfigurationFilePath);
-			}
-			else
-			{
-				FString FullPath;
-				if (ConfigurationFilePath.Contains(TEXT("{Engine}")))
-				{
-					ConfigurationFilePath = FPaths::ConvertRelativePathToFull(ConfigurationFilePath.Replace(TEXT("{Engine}"), *FPaths::EngineDir()));
-				}
+		return;
+	}
 
-				if (!FPaths::IsRelative(ConfigurationFilePath))
-				{
-					FullPath = ConfigurationFilePath;
-				}
-				else
-				{
-					const FString AbsoluteGameDir = FPaths::ConvertRelativePathToFull(FPaths::ProjectDir());
-					FullPath = FPaths::ConvertRelativePathToFull(FPaths::Combine(AbsoluteGameDir, ConfigurationFilePath));
-				}
+	bool bIsBuiltIn = false;
+	FString FilePath = ConfigurationFile.FilePath;
 
-				OCIO_NAMESPACE::ConstConfigRcPtr NewConfig = OCIO_NAMESPACE::Config::CreateFromFile(StringCast<ANSICHAR>(*FullPath).Get());
-				if (NewConfig)
-				{
-					UE_LOG(LogOpenColorIO, Verbose, TEXT("Loaded OCIO configuration file %s"), *FullPath);
-					NativeConfig->Set(NewConfig);
-					
-					StartDirectoryWatch(FullPath);
-				}
-				else
-				{
-					UE_LOG(LogOpenColorIO, Error, TEXT("Could not load OCIO configuration file %s. Verify that the path is good or that the file is valid."), *ConfigurationFile.FilePath);
-				}
-			}
-		}
-#if !PLATFORM_EXCEPTIONS_DISABLED
-		catch (OCIO_NAMESPACE::Exception& exception)
+	if (ConfigurationFile.FilePath.StartsWith(TEXT("ocio://")))
+	{
+		bIsBuiltIn = true;
+	}
+	else if(ConfigurationFile.FilePath.Contains(TEXT("{Engine}")))
+	{
+		FilePath = FPaths::ConvertRelativePathToFull(ConfigurationFile.FilePath.Replace(TEXT("{Engine}"), *FPaths::EngineDir()));
+	}
+	else if (FPaths::IsRelative(ConfigurationFile.FilePath))
+	{
+		const FString AbsoluteGameDir = FPaths::ConvertRelativePathToFull(FPaths::ProjectDir());
+		FilePath = FPaths::ConvertRelativePathToFull(FPaths::Combine(AbsoluteGameDir, ConfigurationFile.FilePath));
+	}
+
+	FOpenColorIOConfigWrapper::FInitializationOptions Options;
+	Options.bAddWorkingColorSpace = true;
+	Config = MakePimpl<FOpenColorIOConfigWrapper>(FilePath, Options);
+
+	if (Config->IsValid())
+	{
+		if (!bIsBuiltIn)
 		{
-			UE_LOG(LogOpenColorIO, Error, TEXT("Could not create OCIO configuration file for %s. Error message: %s."), *ConfigurationFile.FilePath, StringCast<TCHAR>(exception.what()).Get());
+			StartDirectoryWatch(FilePath);
 		}
-#endif
+
+		UE_LOG(LogOpenColorIO, Verbose, TEXT("Loaded OCIO configuration file %s"), *FilePath);
+	}
+	else
+	{
+		UE_LOG(LogOpenColorIO, Error, TEXT("Could not load OCIO configuration file %s. Verify that the path is good or that the file is valid."), *FilePath);
 	}
 #endif
 }
@@ -669,103 +648,45 @@ void UOpenColorIOConfiguration::OnToastCallback(bool bInReloadColorspaces)
 
 #if WITH_EDITOR
 FOpenColorIOEditorConfigurationInspector::FOpenColorIOEditorConfigurationInspector(const UOpenColorIOConfiguration& InConfiguration)
-	: NativeConfig(*InConfiguration.GetNativeConfig_Internal()) // Private implementation pointer is always present and can be safely dereferenced
+	: Configuration(InConfiguration)
 {
 }
 
 int32 FOpenColorIOEditorConfigurationInspector::GetNumColorSpaces() const
 {
-#if WITH_OCIO
-	if (OCIO_NAMESPACE::ConstConfigRcPtr Config = NativeConfig.Get())
-	{
-		return Config->getNumColorSpaces(OCIO_NAMESPACE::SEARCH_REFERENCE_SPACE_ALL, OCIO_NAMESPACE::COLORSPACE_ACTIVE);
-	}
-#endif
-
-	return 0;
+	return Configuration.GetConfigWrapper()->GetNumColorSpaces();
 }
 
 FString FOpenColorIOEditorConfigurationInspector::GetColorSpaceName(int32 Index) const
 {
-#if WITH_OCIO
-	if (OCIO_NAMESPACE::ConstConfigRcPtr Config = NativeConfig.Get())
-	{
-		const char* ColorSpaceName = Config->getColorSpaceNameByIndex(OCIO_NAMESPACE::SEARCH_REFERENCE_SPACE_ALL, OCIO_NAMESPACE::COLORSPACE_ACTIVE, Index);
-
-		return StringCast<TCHAR>(ColorSpaceName).Get();
-	}
-#endif
-
-	return {};
+	return Configuration.GetConfigWrapper()->GetColorSpaceName(Index);
 }
 
 FString FOpenColorIOEditorConfigurationInspector::GetColorSpaceFamilyName(const TCHAR* InColorSpaceName) const
 {
-#if WITH_OCIO
-	if (OCIO_NAMESPACE::ConstConfigRcPtr Config = NativeConfig.Get())
-	{
-		OCIO_NAMESPACE::ConstColorSpaceRcPtr ColorSpace = Config->getColorSpace(StringCast<ANSICHAR>(InColorSpaceName).Get());
-		if (ColorSpace != nullptr)
-		{
-			return StringCast<TCHAR>(ColorSpace->getFamily()).Get();
-		}
-	}
-#endif
-
-	return {};
+	return Configuration.GetConfigWrapper()->GetColorSpaceFamilyName(InColorSpaceName);
 }
 
 int32 FOpenColorIOEditorConfigurationInspector::GetNumDisplays() const
 {
-#if WITH_OCIO
-	if (OCIO_NAMESPACE::ConstConfigRcPtr Config = NativeConfig.Get())
-	{
-		return Config->getNumDisplays();
-	}
-#endif
-
-	return 0;
+	return Configuration.GetConfigWrapper()->GetNumDisplays();
 }
 
 FString FOpenColorIOEditorConfigurationInspector::GetDisplayName(int32 Index) const
 {
-#if WITH_OCIO
-	if (OCIO_NAMESPACE::ConstConfigRcPtr Config = NativeConfig.Get())
-	{
-		const char* DisplayName = Config->getDisplay(Index);
-
-		return StringCast<TCHAR>(DisplayName).Get();
-	}
-#endif
-	
-	return {};
+	return Configuration.GetConfigWrapper()->GetDisplayName(Index);
 }
 
 int32 FOpenColorIOEditorConfigurationInspector::GetNumViews(const TCHAR* InDisplayName) const
 {
-#if WITH_OCIO
-	if (OCIO_NAMESPACE::ConstConfigRcPtr Config = NativeConfig.Get())
-	{
-		return Config->getNumViews(StringCast<ANSICHAR>(InDisplayName).Get());
-	}
-#endif
-
-	return 0;
+	return Configuration.GetConfigWrapper()->GetNumViews(InDisplayName);
 }
 
 FString FOpenColorIOEditorConfigurationInspector::GetViewName(const TCHAR* InDisplayName, int32 Index) const
 {
-#if WITH_OCIO
-	if (OCIO_NAMESPACE::ConstConfigRcPtr Config = NativeConfig.Get())
-	{
-		const char* ViewName = Config->getView(StringCast<ANSICHAR>(InDisplayName).Get(), Index);
-
-		return StringCast<TCHAR>(ViewName).Get();
-	}
-#endif
-	
-	return {};
+	return Configuration.GetConfigWrapper()->GetViewName(InDisplayName, Index);
 }
 #endif //WITH_EDITOR
 
+
 #undef LOCTEXT_NAMESPACE
diff --git a/Engine/Plugins/Compositing/OpenColorIO/Source/OpenColorIO/Private/OpenColorIOLibHandler.cpp b/Engine/Plugins/Compositing/OpenColorIO/Source/OpenColorIO/Private/OpenColorIOLibHandler.cpp
deleted file mode 100644
diff --git a/Engine/Plugins/Compositing/OpenColorIO/Source/OpenColorIO/Private/OpenColorIOLibHandler.h b/Engine/Plugins/Compositing/OpenColorIO/Source/OpenColorIO/Private/OpenColorIOLibHandler.h
deleted file mode 100644
diff --git a/Engine/Plugins/Compositing/OpenColorIO/Source/OpenColorIO/Private/OpenColorIOModule.cpp b/Engine/Plugins/Compositing/OpenColorIO/Source/OpenColorIO/Private/OpenColorIOModule.cpp
--- a/Engine/Plugins/Compositing/OpenColorIO/Source/OpenColorIO/Private/OpenColorIOModule.cpp
+++ b/Engine/Plugins/Compositing/OpenColorIO/Source/OpenColorIO/Private/OpenColorIOModule.cpp
@@ -3,9 +3,7 @@
 #include "OpenColorIOModule.h"
 
 #include "Interfaces/IPluginManager.h"
-#include "OpenColorIOLibHandler.h"
 #include "OpenColorIODisplayManager.h"
-#include "OpenColorIONativeConfiguration.h"
 #include "ShaderCore.h"
 
 
@@ -21,8 +19,6 @@ FOpenColorIOModule::FOpenColorIOModule()
 
 void FOpenColorIOModule::StartupModule()
 {
-	bInitializedLib = FOpenColorIOLibHandler::Initialize();
-
 	// Maps virtual shader source directory /Plugin/OpenCVLensDistortion to the plugin's actual Shaders directory.
 	FString PluginShaderDir = FPaths::Combine(IPluginManager::Get().FindPlugin(TEXT("OpenColorIO"))->GetBaseDir(), TEXT("Shaders"));
 	AddShaderSourceDirectoryMapping(TEXT("/Plugin/OpenColorIO"), PluginShaderDir);
@@ -30,28 +26,13 @@ void FOpenColorIOModule::StartupModule()
 
 void FOpenColorIOModule::ShutdownModule()
 {
-	FOpenColorIOLibHandler::Shutdown();
+
 }
 
 FOpenColorIODisplayManager& FOpenColorIOModule::GetDisplayManager()
 {
 	return *DisplayManager;
 }
-
-FOpenColorIONativeInterchangeConfiguration* FOpenColorIOModule::GetNativeInterchangeConfig_Internal()
-{
-	if (bInitializedLib && NativeInterchangeConfig == nullptr)
-	{
-		/**
-		 * Create a single config used for conversions between the working color spaceand the interchange space.
-		 * Note that we delay the creation here to ensure that the global working color space has been loaded by
-		 * the engine via renderer settings.
-		 */
-		NativeInterchangeConfig = MakeUnique<FOpenColorIONativeInterchangeConfiguration>();
-	}
-
-	return NativeInterchangeConfig.Get();
-}
 	
 IMPLEMENT_MODULE(FOpenColorIOModule, OpenColorIO);
 
diff --git a/Engine/Plugins/Compositing/OpenColorIO/Source/OpenColorIO/Private/OpenColorIOModule.h b/Engine/Plugins/Compositing/OpenColorIO/Source/OpenColorIO/Private/OpenColorIOModule.h
--- a/Engine/Plugins/Compositing/OpenColorIO/Source/OpenColorIO/Private/OpenColorIOModule.h
+++ b/Engine/Plugins/Compositing/OpenColorIO/Source/OpenColorIO/Private/OpenColorIOModule.h
@@ -13,7 +13,6 @@ DECLARE_LOG_CATEGORY_EXTERN(LogOpenColorIO, Log, All);
 
 
 class FOpenColorIODisplayManager;
-class FOpenColorIONativeInterchangeConfiguration;
 
 /**
  * Implements the OpenColorIO module.
@@ -32,13 +31,10 @@ public:
 
 	//~ Begin IOpenColorIOModule
 	virtual FOpenColorIODisplayManager& GetDisplayManager() override;
-	virtual FOpenColorIONativeInterchangeConfiguration* GetNativeInterchangeConfig_Internal() override;
 	//~ End IOpenColorIOModule
 
 private:
 
-	bool bInitializedLib = false;
 	TUniquePtr<FOpenColorIODisplayManager> DisplayManager;
-	TUniquePtr<FOpenColorIONativeInterchangeConfiguration> NativeInterchangeConfig;
 };
 
diff --git a/Engine/Plugins/Compositing/OpenColorIO/Source/OpenColorIO/Private/OpenColorIONativeConfiguration.cpp b/Engine/Plugins/Compositing/OpenColorIO/Source/OpenColorIO/Private/OpenColorIONativeConfiguration.cpp
deleted file mode 100644
diff --git a/Engine/Plugins/Compositing/OpenColorIO/Source/OpenColorIO/Private/OpenColorIONativeConfiguration.h b/Engine/Plugins/Compositing/OpenColorIO/Source/OpenColorIO/Private/OpenColorIONativeConfiguration.h
deleted file mode 100644
diff --git a/Engine/Plugins/Compositing/OpenColorIO/Source/OpenColorIO/Private/OpenColorIOShaderMap.cpp b/Engine/Plugins/Compositing/OpenColorIO/Source/OpenColorIO/Private/OpenColorIOShaderMap.cpp
--- a/Engine/Plugins/Compositing/OpenColorIO/Source/OpenColorIO/Private/OpenColorIOShaderMap.cpp
+++ b/Engine/Plugins/Compositing/OpenColorIO/Source/OpenColorIO/Private/OpenColorIOShaderMap.cpp
@@ -15,14 +15,10 @@
 #include "UObject/UObjectGlobals.h"
 
 #if WITH_EDITOR
+	#include "OpenColorIOWrapper.h"
 	#include "DerivedDataCacheInterface.h"
 	#include "Interfaces/ITargetPlatformManagerModule.h"
 	#include "TickableEditorObject.h"
-
-#if WITH_OCIO
-	#include "OpenColorIO/OpenColorIO.h"
-#endif //WITH_OCIO
-
 #endif
 
 
@@ -234,12 +230,10 @@ FOpenColorIOShaderMap* FOpenColorIOShaderMap::FindId(const FOpenColorIOShaderMap
 
 void OpenColorIOShaderMapAppendKeyString(EShaderPlatform InPlatform, FString& OutKeyString)
 {
-#if WITH_OCIO
 	//Keep library version in the DDC key to invalidate it once we move to a new library
 	OutKeyString += TEXT("OCIOVersion");
-	OutKeyString += TEXT(OCIO_VERSION);
+	OutKeyString += OpenColorIOWrapper::GetVersion();
 	OutKeyString += TEXT("_");
-#endif // WITH_OCIO
 }
 
 /** Creates a string key for the derived data cache given a shader map id. */
diff --git a/Engine/Plugins/Compositing/OpenColorIO/Source/OpenColorIO/Public/IOpenColorIOModule.h b/Engine/Plugins/Compositing/OpenColorIO/Source/OpenColorIO/Public/IOpenColorIOModule.h
--- a/Engine/Plugins/Compositing/OpenColorIO/Source/OpenColorIO/Public/IOpenColorIOModule.h
+++ b/Engine/Plugins/Compositing/OpenColorIO/Source/OpenColorIO/Public/IOpenColorIOModule.h
@@ -45,8 +45,6 @@ public:
 	 */
 	virtual FOpenColorIODisplayManager& GetDisplayManager() = 0;
 
-	/**
-	 * Internal only: Returns a minimal dynamically-created native config for conversions between interchange and working color spaces.
-	 */
-	virtual FOpenColorIONativeInterchangeConfiguration* GetNativeInterchangeConfig_Internal() = 0;
+	UE_DEPRECATED(5.3, "This function is deprecated and has been replaced by FOpenColorIOConfigWrapper class in the OpenColorIOWrapper module.")
+	FOpenColorIONativeInterchangeConfiguration* GetNativeInterchangeConfig_Internal() { return nullptr; };
 };
\ No newline at end of file
diff --git a/Engine/Plugins/Compositing/OpenColorIO/Source/OpenColorIO/Public/OpenColorIOConfiguration.h b/Engine/Plugins/Compositing/OpenColorIO/Source/OpenColorIO/Public/OpenColorIOConfiguration.h
--- a/Engine/Plugins/Compositing/OpenColorIO/Source/OpenColorIO/Public/OpenColorIOConfiguration.h
+++ b/Engine/Plugins/Compositing/OpenColorIO/Source/OpenColorIO/Public/OpenColorIOConfiguration.h
@@ -14,7 +14,7 @@
 #include "OpenColorIOConfiguration.generated.h"
 
 
-class FOpenColorIONativeConfiguration;
+class FOpenColorIOConfigWrapper;
 class FOpenColorIOTransformResource;
 class FTextureResource;
 class UOpenColorIOColorTransform;
@@ -37,8 +37,6 @@ public:
 
 public:
 
-	static constexpr TCHAR WorkingColorSpaceName[] = TEXT("Working Color Space");
-
 	bool IsTransformReady(const FOpenColorIOColorConversionSettings& InSettings);
 	bool GetRenderResources(ERHIFeatureLevel::Type InFeatureLevel, const FOpenColorIOColorConversionSettings& InSettings, FOpenColorIOTransformResource*& OutShaderResource, TSortedMap<int32, FTextureResource*>& OutTextureResources);
 	bool HasTransform(const FString& InSourceColorSpace, const FString& InDestinationColorSpace);
@@ -65,8 +63,9 @@ public:
 	*/
 	void ConfigPathChangedEvent(const TArray<FFileChangeData>& InFileChanges, const FString InFileMountPath);
 
-	/** Internal only: Replacement for previous `GetLoadedConfiguration()` and `GetLoadedConfigurationFile()` functions, returning the private implementation of the native OCIO config. */
-	FOpenColorIONativeConfiguration* GetNativeConfig_Internal() const;
+	/** Internal only: Replacement for previous `GetNativeConfig_Internal()`, `GetLoadedConfiguration()` and `GetLoadedConfigurationFile()` functions, returning the private implementation of the native OCIO config. */
+	FOpenColorIOConfigWrapper* GetConfigWrapper() const;
+
 protected:
 
 	const TObjectPtr<UOpenColorIOColorTransform>* FindTransform(const FOpenColorIOColorConversionSettings& InSettings) const;
@@ -136,39 +135,43 @@ private:
 	/** Information about the currently watched directory. Helps us manage the directory change events. */
 	FOCIOConfigWatchedDirInfo WatchedDirectoryInfo;
 
-	/** Private implementation of the native OpenColorIO config object. */
-	TPimplPtr<FOpenColorIONativeConfiguration> NativeConfig;
+#if WITH_EDITORONLY_DATA
+	/** Private implementation of the OpenColorIO config object. */
+	TPimplPtr<FOpenColorIOConfigWrapper> Config = nullptr;
+#endif
 };
 
 
 #if WITH_EDITOR
-/**
- * Editor-only class to access native configuration names of color spaces, displays and views for menus.
- */
 class FOpenColorIOEditorConfigurationInspector {
 public:
 
+	UE_DEPRECATED(5.3, "This class is deprecated and has been replaced by FOpenColorIOConfigWrapper class in the OpenColorIOWrapper module.")
 	OPENCOLORIO_API FOpenColorIOEditorConfigurationInspector(const UOpenColorIOConfiguration& InConfiguration);
 
-	/** Get the number of color spaces in the configuration. */
+	UE_DEPRECATED(5.3, "This class is deprecated and has been replaced by FOpenColorIOConfigWrapper class in the OpenColorIOWrapper module.")
 	OPENCOLORIO_API int32 GetNumColorSpaces() const;
-	/** Get a color space name at an index. */
+	
+	UE_DEPRECATED(5.3, "This class is deprecated and has been replaced by FOpenColorIOConfigWrapper class in the OpenColorIOWrapper module.")
 	OPENCOLORIO_API FString GetColorSpaceName(int32 Index) const;
-	/** Get the family name for a color space. */
+	
+	UE_DEPRECATED(5.3, "This class is deprecated and has been replaced by FOpenColorIOConfigWrapper class in the OpenColorIOWrapper module.")
 	OPENCOLORIO_API FString GetColorSpaceFamilyName(const TCHAR* InColorSpaceName) const;
 
-	/** Get the number of displays in the configuration. */
+	UE_DEPRECATED(5.3, "This class is deprecated and has been replaced by FOpenColorIOConfigWrapper class in the OpenColorIOWrapper module.")
 	OPENCOLORIO_API int32 GetNumDisplays() const;
-	/** Get a display name at an index. */
+	
+	UE_DEPRECATED(5.3, "This class is deprecated and has been replaced by FOpenColorIOConfigWrapper class in the OpenColorIOWrapper module.")
 	OPENCOLORIO_API FString GetDisplayName(int32 Index) const;
 
-	/** Get the number of views for a display. */
+	UE_DEPRECATED(5.3, "This class is deprecated and has been replaced by FOpenColorIOConfigWrapper class in the OpenColorIOWrapper module.")
 	OPENCOLORIO_API int32 GetNumViews(const TCHAR* InDisplayName) const;
-	/** Get a view name for its display and index. */
+	
+	UE_DEPRECATED(5.3, "This class is deprecated and has been replaced by FOpenColorIOConfigWrapper class in the OpenColorIOWrapper module.")
 	OPENCOLORIO_API FString GetViewName(const TCHAR* InDisplayName, int32 Index) const;
 
 private:
-	/** Reference to the native OCIO config. */
-	const FOpenColorIONativeConfiguration& NativeConfig;
+
+	const UOpenColorIOConfiguration& Configuration;
 };
 #endif //WITH_EDITOR
diff --git a/Engine/Plugins/Compositing/OpenColorIO/Source/OpenColorIO/Public/OpenColorIOShader.h b/Engine/Plugins/Compositing/OpenColorIO/Source/OpenColorIO/Public/OpenColorIOShader.h
--- a/Engine/Plugins/Compositing/OpenColorIO/Source/OpenColorIO/Public/OpenColorIOShader.h
+++ b/Engine/Plugins/Compositing/OpenColorIO/Source/OpenColorIO/Public/OpenColorIOShader.h
@@ -64,8 +64,11 @@ END_SHADER_PARAMETER_STRUCT()
 
 namespace OpenColorIOShader
 {
-	static constexpr TCHAR OpenColorIOShaderFunctionName[] = TEXT("OCIOConvert");
-	static constexpr uint32 Lut3dEdgeLength = 65;
+/** Fields moved to OpenColorIOWrapper.h
+ *
+ *	static constexpr TCHAR OpenColorIOShaderFunctionName[] = TEXT("OCIOConvert");
+ *	static constexpr uint32 Lut3dEdgeLength = 65;
+*/
 
 	// Combined maximum of either Ocio_lut1d or Ocio_lut3d, since the generated shaders will never use the same slots for both types.
 	static constexpr uint32 MaximumTextureSlots = 8;
diff --git a/Engine/Plugins/Compositing/OpenColorIO/Source/OpenColorIO/Public/OpenColorIOShared.h b/Engine/Plugins/Compositing/OpenColorIO/Source/OpenColorIO/Public/OpenColorIOShared.h
--- a/Engine/Plugins/Compositing/OpenColorIO/Source/OpenColorIO/Public/OpenColorIOShared.h
+++ b/Engine/Plugins/Compositing/OpenColorIO/Source/OpenColorIO/Public/OpenColorIOShared.h
@@ -24,11 +24,9 @@
 #include "Templates/RefCounting.h"
 #include "OpenColorIOShaderCompilationManager.h"
 
-class FOpenColorIOTransformResource;
-class FOpenColorIOShaderMap;
-class FOpenColorIOPixelShader;
-class FOpenColorIOShaderMapId;
-
+#if WITH_EDITOR
+#include "OpenColorIOWrapperDefines.h"
+#else
 /** Enum used to indicate whether the working color space should be used as a source or destination. */
 enum class EOpenColorIOWorkingColorSpaceTransform : uint8
 {
@@ -36,6 +34,13 @@ enum class EOpenColorIOWorkingColorSpaceTransform : uint8
 	Source = 1,
 	Destination = 2
 };
+#endif
+
+class FOpenColorIOTransformResource;
+class FOpenColorIOShaderMap;
+class FOpenColorIOPixelShader;
+class FOpenColorIOShaderMapId;
+
 
 /** Stores outputs from the color transform compile that need to be saved. */
 class FOpenColorIOCompilationOutput
diff --git a/Engine/Plugins/Compositing/OpenColorIO/Source/OpenColorIOEditor/OpenColorIOEditor.Build.cs b/Engine/Plugins/Compositing/OpenColorIO/Source/OpenColorIOEditor/OpenColorIOEditor.Build.cs
--- a/Engine/Plugins/Compositing/OpenColorIO/Source/OpenColorIOEditor/OpenColorIOEditor.Build.cs
+++ b/Engine/Plugins/Compositing/OpenColorIO/Source/OpenColorIOEditor/OpenColorIOEditor.Build.cs
@@ -15,8 +15,8 @@ namespace UnrealBuildTool.Rules
 					"EditorFramework",
 					"Engine",
 					"LevelEditor",
+					"OpenColorIOWrapper",
 					"OpenColorIO",
-					"OpenColorIOLib",
 					"Projects",
 					"PropertyEditor",
 					"Renderer",
diff --git a/Engine/Plugins/Compositing/OpenColorIO/Source/OpenColorIOEditor/Private/OpenColorIOColorSpaceCustomization.cpp b/Engine/Plugins/Compositing/OpenColorIO/Source/OpenColorIOEditor/Private/OpenColorIOColorSpaceCustomization.cpp
--- a/Engine/Plugins/Compositing/OpenColorIO/Source/OpenColorIOEditor/Private/OpenColorIOColorSpaceCustomization.cpp
+++ b/Engine/Plugins/Compositing/OpenColorIO/Source/OpenColorIOEditor/Private/OpenColorIOColorSpaceCustomization.cpp
@@ -7,7 +7,7 @@
 #include "IPropertyUtilities.h"
 #include "PropertyHandle.h"
 #include "OpenColorIOConfiguration.h"
-#include "OpenColorIO/OpenColorIO.h"
+#include "OpenColorIOWrapper.h"
 #include "Widgets/Input/SComboButton.h"
 
 #define LOCTEXT_NAMESPACE "OpenColorIOColorSpaceCustomization"
@@ -15,8 +15,6 @@
 void FOpenColorIOColorSpaceCustomization::CustomizeHeader(TSharedRef<IPropertyHandle> InPropertyHandle, FDetailWidgetRow& HeaderRow, IPropertyTypeCustomizationUtils& CustomizationUtils)
 {
 	//Reset internals
-	ConfigInspector.Reset(nullptr);
-
 	CachedProperty = InPropertyHandle;
 
 	if (CachedProperty->GetNumPerObjectValues() == 1 && CachedProperty->IsValidHandle())
@@ -65,8 +63,6 @@ void FOpenColorIOColorSpaceCustomization::CustomizeHeader(TSharedRef<IPropertyHa
 void FOpenColorIODisplayViewCustomization::CustomizeHeader(TSharedRef<IPropertyHandle> InPropertyHandle, FDetailWidgetRow& HeaderRow, IPropertyTypeCustomizationUtils& CustomizationUtils)
 {
 	//Reset internals
-	ConfigInspector.Reset(nullptr);
-
 	CachedProperty = InPropertyHandle;
 
 	if (CachedProperty->GetNumPerObjectValues() == 1 && CachedProperty->IsValidHandle())
@@ -117,7 +113,7 @@ IPropertyTypeCustomizationOpenColorIO::IPropertyTypeCustomizationOpenColorIO(TSh
 {
 }
 
-TUniquePtr<FOpenColorIOEditorConfigurationInspector> IPropertyTypeCustomizationOpenColorIO::GetConfigurationInspector(const TSharedPtr<IPropertyHandle>& InConfigurationAssetProperty)
+FOpenColorIOConfigWrapper* IPropertyTypeCustomizationOpenColorIO::GetConfigWrapper(const TSharedPtr<IPropertyHandle>& InConfigurationAssetProperty)
 {
 	if (InConfigurationAssetProperty.IsValid())
 	{
@@ -128,7 +124,7 @@ TUniquePtr<FOpenColorIOEditorConfigurationInspector> IPropertyTypeCustomizationO
 		UOpenColorIOConfiguration* Configuration = reinterpret_cast<UOpenColorIOConfiguration*>(RawData[0]);
 		check(Configuration);
 		
-		return MakeUnique<FOpenColorIOEditorConfigurationInspector>(*Configuration);
+		return Configuration->GetConfigWrapper();
 	}
 
 	return nullptr;
@@ -169,19 +165,25 @@ void FOpenColorIOColorSpaceCustomization::PopulateSubMenu(FMenuBuilder& InMenuBu
 	// To keep track of submenus that were already added
 	TArray<FString> ExistingSubMenus;
 
-	for (int32 i = 0; i < ConfigInspector->GetNumColorSpaces(); ++i)
+	const FOpenColorIOConfigWrapper* ConfigWrapper = GetConfigWrapper(ConfigurationObjectProperty);
+
+	if (ConfigWrapper)
 	{
-		FOpenColorIOColorSpace ColorSpace;
-		ColorSpace.ColorSpaceIndex = i;
-		ColorSpace.ColorSpaceName = ConfigInspector->GetColorSpaceName(i);
-		ColorSpace.FamilyName = ConfigInspector->GetColorSpaceFamilyName(*ColorSpace.ColorSpaceName);
-		
-		//Filter out color spaces that don't belong to this hierarchy
-		if (InPreviousFamilyHierarchy.IsEmpty() || ColorSpace.FamilyName.Contains(InPreviousFamilyHierarchy))
+		for (int32 i = 0; i < ConfigWrapper->GetNumColorSpaces(); ++i)
 		{
-			ProcessColorSpaceForMenuGeneration(InMenuBuilder, InMenuDepth, InPreviousFamilyHierarchy, ColorSpace, ExistingSubMenus);
+			FOpenColorIOColorSpace ColorSpace;
+			ColorSpace.ColorSpaceIndex = i;
+			ColorSpace.ColorSpaceName = ConfigWrapper->GetColorSpaceName(i);
+			ColorSpace.FamilyName = ConfigWrapper->GetColorSpaceFamilyName(*ColorSpace.ColorSpaceName);
+
+			//Filter out color spaces that don't belong to this hierarchy
+			if (InPreviousFamilyHierarchy.IsEmpty() || ColorSpace.FamilyName.Contains(InPreviousFamilyHierarchy))
+			{
+				ProcessColorSpaceForMenuGeneration(InMenuBuilder, InMenuDepth, InPreviousFamilyHierarchy, ColorSpace, ExistingSubMenus);
+			}
 		}
 	}
+
 }
 
 void FOpenColorIOColorSpaceCustomization::AddMenuEntry(FMenuBuilder& InMenuBuilder, const FOpenColorIOColorSpace& InColorSpace)
@@ -221,7 +223,7 @@ void FOpenColorIOColorSpaceCustomization::AddMenuEntry(FMenuBuilder& InMenuBuild
 
 TSharedRef<SWidget> FOpenColorIOColorSpaceCustomization::HandleSourceComboButtonMenuContent()
 {
-	ConfigInspector = GetConfigurationInspector(ConfigurationObjectProperty);
+	const FOpenColorIOConfigWrapper* ConfigWrapper = GetConfigWrapper(ConfigurationObjectProperty);
 
 	// Generate menu
 	FMenuBuilder MenuBuilder(true, nullptr);
@@ -229,16 +231,16 @@ TSharedRef<SWidget> FOpenColorIOColorSpaceCustomization::HandleSourceComboButton
 
 	MenuBuilder.BeginSection("AllColorSpaces", LOCTEXT("AllColorSpacesSection", "ColorSpaces"));
 	{
-		if (ConfigInspector)
+		if (ConfigWrapper)
 		{
-			const int32 ColorSpaceCount = ConfigInspector->GetNumColorSpaces();
+			const int32 ColorSpaceCount = ConfigWrapper->GetNumColorSpaces();
 
 			for (int32 i = 0; i < ColorSpaceCount; ++i)
 			{
 				FOpenColorIOColorSpace ColorSpace;
 				ColorSpace.ColorSpaceIndex = i;
-				ColorSpace.ColorSpaceName = ConfigInspector->GetColorSpaceName(i);
-				ColorSpace.FamilyName = ConfigInspector->GetColorSpaceFamilyName(*ColorSpace.ColorSpaceName);
+				ColorSpace.ColorSpaceName = ConfigWrapper->GetColorSpaceName(i);
+				ColorSpace.FamilyName = ConfigWrapper->GetColorSpaceFamilyName(*ColorSpace.ColorSpaceName);
 
 				//Top level menus have no preceding hierarchy.
 				const int32 CurrentMenuDepth = 0;
@@ -263,11 +265,16 @@ TSharedRef<SWidget> FOpenColorIOColorSpaceCustomization::HandleSourceComboButton
 
 void FOpenColorIODisplayViewCustomization::PopulateViewSubMenu(FMenuBuilder& InMenuBuilder, FOpenColorIODisplayView InDisplayView)
 {
-	for (int32 i = 0; i < ConfigInspector->GetNumViews(*InDisplayView.Display); ++i)
+	const FOpenColorIOConfigWrapper* ConfigWrapper = GetConfigWrapper(ConfigurationObjectProperty);
+
+	if (ConfigWrapper)
 	{
-		InDisplayView.View = ConfigInspector->GetViewName(*InDisplayView.Display, i);
-		
-		AddMenuEntry(InMenuBuilder, InDisplayView);
+		for (int32 i = 0; i < ConfigWrapper->GetNumViews(*InDisplayView.Display); ++i)
+		{
+			InDisplayView.View = ConfigWrapper->GetViewName(*InDisplayView.Display, i);
+
+			AddMenuEntry(InMenuBuilder, InDisplayView);
+		}
 	}
 }
 
@@ -308,7 +315,7 @@ void FOpenColorIODisplayViewCustomization::AddMenuEntry(FMenuBuilder& InMenuBuil
 
 TSharedRef<SWidget> FOpenColorIODisplayViewCustomization::HandleSourceComboButtonMenuContent()
 {
-	ConfigInspector = GetConfigurationInspector(ConfigurationObjectProperty);
+	const FOpenColorIOConfigWrapper* ConfigWrapper = GetConfigWrapper(ConfigurationObjectProperty);
 
 	// Generate menu
 	FMenuBuilder MenuBuilder(true, nullptr);
@@ -316,13 +323,13 @@ TSharedRef<SWidget> FOpenColorIODisplayViewCustomization::HandleSourceComboButto
 
 	MenuBuilder.BeginSection("AllDisplayViews", LOCTEXT("AllDisplayViewsSection", "Display - View"));
 	{
-		if (ConfigInspector)
+		if (ConfigWrapper)
 		{
-			const int32 DisplayCount = ConfigInspector->GetNumDisplays();
+			const int32 DisplayCount = ConfigWrapper->GetNumDisplays();
 
 			for (int32 i = 0; i < DisplayCount; ++i)
 			{
-				const FOpenColorIODisplayView DisplayViewValue = { ConfigInspector->GetDisplayName(i), TEXT("<Invalid>")};
+				const FOpenColorIODisplayView DisplayViewValue = { ConfigWrapper->GetDisplayName(i), TEXT("<Invalid>")};
 
 				MenuBuilder.AddSubMenu(
 					FText::FromString(DisplayViewValue.Display),
diff --git a/Engine/Plugins/Compositing/OpenColorIO/Source/OpenColorIOEditor/Private/OpenColorIOColorSpaceCustomization.h b/Engine/Plugins/Compositing/OpenColorIO/Source/OpenColorIOEditor/Private/OpenColorIOColorSpaceCustomization.h
--- a/Engine/Plugins/Compositing/OpenColorIO/Source/OpenColorIOEditor/Private/OpenColorIOColorSpaceCustomization.h
+++ b/Engine/Plugins/Compositing/OpenColorIO/Source/OpenColorIOEditor/Private/OpenColorIOColorSpaceCustomization.h
@@ -9,7 +9,7 @@
 #include "OpenColorIOColorSpace.h"
 #include "Widgets/SWidget.h"
 
-class FOpenColorIOEditorConfigurationInspector;
+class FOpenColorIOConfigWrapper;
 
 /**
  * Implements a details view customization for the FOpenColorIOConfiguration
@@ -25,7 +25,7 @@ public:
 
 protected:
 
-	static TUniquePtr<FOpenColorIOEditorConfigurationInspector> GetConfigurationInspector(const TSharedPtr<IPropertyHandle>& InConfigurationObjectProperty);
+	static FOpenColorIOConfigWrapper* GetConfigWrapper(const TSharedPtr<IPropertyHandle>& InConfigurationObjectProperty);
 
 protected:
 
@@ -34,9 +34,6 @@ protected:
 
 	/** Pointer to the configuration object property handle. */
 	TSharedPtr<IPropertyHandle> ConfigurationObjectProperty;
-
-	/** Raw pointer to the native OCIO config. */
-	TUniquePtr<FOpenColorIOEditorConfigurationInspector> ConfigInspector;
 };
 
 
diff --git a/Engine/Plugins/Compositing/OpenColorIO/Source/OpenColorIOWrapper/OpenColorIOWrapper.Build.cs b/Engine/Plugins/Compositing/OpenColorIO/Source/OpenColorIOWrapper/OpenColorIOWrapper.Build.cs
new file mode 100644
--- /dev/null
+++ b/Engine/Plugins/Compositing/OpenColorIO/Source/OpenColorIOWrapper/OpenColorIOWrapper.Build.cs
@@ -0,0 +1,29 @@
+// Copyright Epic Games, Inc. All Rights Reserved.
+
+namespace UnrealBuildTool.Rules
+{
+	public class OpenColorIOWrapper : ModuleRules
+	{
+		public OpenColorIOWrapper(ReadOnlyTargetRules Target) : base(Target)
+		{
+			PCHUsage = PCHUsageMode.UseExplicitOrSharedPCHs;
+
+			PrivateIncludePathModuleNames.AddRange(new string[]
+			{
+				"Engine", // for TextureDefines.h
+			});
+
+			PrivateDependencyModuleNames.AddRange(new string[]
+			{
+				"ColorManagement",
+				"OpenColorIOLib" // The OpenColorIO third-party library will only be available on desktop editor.
+			});
+
+			PublicDependencyModuleNames.AddRange(new string[]
+			{
+				"Core",
+				"ImageCore"
+			});
+		}
+	}
+}
diff --git a/Engine/Plugins/Compositing/OpenColorIO/Source/OpenColorIOWrapper/Private/OpenColorIOLibHandler.cpp b/Engine/Plugins/Compositing/OpenColorIO/Source/OpenColorIOWrapper/Private/OpenColorIOLibHandler.cpp
new file mode 100644
--- /dev/null
+++ b/Engine/Plugins/Compositing/OpenColorIO/Source/OpenColorIOWrapper/Private/OpenColorIOLibHandler.cpp
@@ -0,0 +1,66 @@
+// Copyright Epic Games, Inc. All Rights Reserved.
+
+#include "OpenColorIOLibHandler.h"
+
+#include "HAL/PlatformProcess.h"
+#include "Logging/LogMacros.h"
+#include "Misc/Paths.h"
+#include "OpenColorIOWrapperModule.h"
+
+
+ //~ Static initialization
+ //--------------------------------------------------------------------
+void* FOpenColorIOLibHandler::LibHandle = nullptr;
+
+
+//~ Initialization functions implementation
+//--------------------------------------------------------------------
+bool FOpenColorIOLibHandler::Initialize()
+{
+#if WITH_OCIO && defined(OCIO_DLL_NAME)
+	check(LibHandle == nullptr);
+
+	const FString OCIOBinPath = FPaths::Combine(FPaths::EngineDir(), TEXT("Binaries/ThirdParty/OpenColorIO"), FPlatformProcess::GetBinariesSubdirectory());
+	const FString DLLPath = FPaths::Combine(OCIOBinPath, TEXT(PREPROCESSOR_TO_STRING(OCIO_DLL_NAME)));
+
+	FPlatformProcess::PushDllDirectory(*OCIOBinPath);
+
+	if (!FPaths::FileExists(DLLPath))
+	{
+		UE_LOG(LogOpenColorIOWrapper, Error, TEXT("Failed to find the OpenColorIO dll. Plug-in will not be functional."));
+		return false;
+	}
+	
+	LibHandle = FPlatformProcess::GetDllHandle(*DLLPath);
+	FPlatformProcess::PopDllDirectory(*OCIOBinPath);
+
+	if (LibHandle == nullptr)
+	{
+		UE_LOG(LogOpenColorIOWrapper, Error, TEXT("Failed to load required library %s. Plug-in will not be functional."), *DLLPath);
+		return false;
+	}
+
+#else
+	return false;
+#endif // WITH_EDITOR
+	return true;
+}
+
+bool FOpenColorIOLibHandler::IsInitialized()
+{
+	return (LibHandle != nullptr);
+}
+
+void FOpenColorIOLibHandler::Shutdown()
+{
+#if WITH_OCIO && defined(OCIO_DLL_NAME)
+	if (LibHandle != nullptr)
+	{
+		FPlatformProcess::FreeDllHandle(LibHandle);
+		LibHandle = nullptr;
+	}
+#endif // WITH_EDITOR
+}
+
+
+
diff --git a/Engine/Plugins/Compositing/OpenColorIO/Source/OpenColorIOWrapper/Private/OpenColorIOLibHandler.h b/Engine/Plugins/Compositing/OpenColorIO/Source/OpenColorIOWrapper/Private/OpenColorIOLibHandler.h
new file mode 100644
--- /dev/null
+++ b/Engine/Plugins/Compositing/OpenColorIO/Source/OpenColorIOWrapper/Private/OpenColorIOLibHandler.h
@@ -0,0 +1,15 @@
+// Copyright Epic Games, Inc. All Rights Reserved.
+
+#pragma once
+
+
+class FOpenColorIOLibHandler
+{
+public:
+	static bool Initialize();
+	static bool IsInitialized();
+	static void Shutdown();
+
+private:
+	static void* LibHandle;
+};
diff --git a/Engine/Plugins/Compositing/OpenColorIO/Source/OpenColorIOWrapper/Private/OpenColorIOWrapper.cpp b/Engine/Plugins/Compositing/OpenColorIO/Source/OpenColorIOWrapper/Private/OpenColorIOWrapper.cpp
new file mode 100644
--- /dev/null
+++ b/Engine/Plugins/Compositing/OpenColorIO/Source/OpenColorIOWrapper/Private/OpenColorIOWrapper.cpp
@@ -0,0 +1,728 @@
+// Copyright Epic Games, Inc. All Rights Reserved.
+
+#include "OpenColorIOWrapper.h"
+#include "OpenColorIOWrapperModule.h"
+
+#include "Containers/Map.h"
+#include "ColorSpace.h"
+#include "Engine/TextureDefines.h"
+#include "ImageCore.h"
+
+#if WITH_OCIO
+THIRD_PARTY_INCLUDES_START
+#include "OpenColorIO/OpenColorIO.h"
+THIRD_PARTY_INCLUDES_END
+#endif
+
+#define LOCTEXT_NAMESPACE "OpenColorIOWrapper"
+
+const TCHAR* OpenColorIOWrapper::GetVersion()
+{
+#if WITH_OCIO
+	return TEXT(OCIO_VERSION);
+#else
+	return TEXT("");
+#endif
+}
+
+struct FOpenColorIOConfigPimpl
+{
+#if WITH_OCIO
+	OCIO_NAMESPACE::ConstConfigRcPtr Config = nullptr;
+#endif
+};
+
+struct FOpenColorIOProcessorPimpl
+{
+#if WITH_OCIO
+	OCIO_NAMESPACE::ConstProcessorRcPtr Processor = nullptr;
+
+	/** Get processor optimization flags. */
+	static OCIO_NAMESPACE::OptimizationFlags GetOptimizationFlags()
+	{
+		return static_cast<OCIO_NAMESPACE::OptimizationFlags>(
+			OCIO_NAMESPACE::OptimizationFlags::OPTIMIZATION_DEFAULT |
+			OCIO_NAMESPACE::OptimizationFlags::OPTIMIZATION_NO_DYNAMIC_PROPERTIES
+		);
+	}
+#endif
+};
+
+struct FOpenColorIOGPUProcessorPimpl
+{
+#if WITH_OCIO
+	OCIO_NAMESPACE::ConstGPUProcessorRcPtr Processor = nullptr;
+	OCIO_NAMESPACE::GpuShaderDescRcPtr ShaderDescription = nullptr;
+#endif
+};
+
+#if WITH_OCIO
+namespace {
+	TUniquePtr<OCIO_NAMESPACE::PackedImageDesc> GetImageDesc(const FImageView& InImage)
+	{
+		OCIO_NAMESPACE::ChannelOrdering Ordering;
+		OCIO_NAMESPACE::BitDepth BitDepth;
+
+		switch (InImage.Format)
+		{
+		case ERawImageFormat::BGRA8:
+			Ordering = OCIO_NAMESPACE::CHANNEL_ORDERING_BGRA;
+			BitDepth = OCIO_NAMESPACE::BIT_DEPTH_UINT8;
+			break;
+		case ERawImageFormat::RGBA16:
+			Ordering = OCIO_NAMESPACE::CHANNEL_ORDERING_RGBA;
+			BitDepth = OCIO_NAMESPACE::BIT_DEPTH_UINT16;
+			break;
+		case ERawImageFormat::RGBA16F:
+			Ordering = OCIO_NAMESPACE::CHANNEL_ORDERING_RGBA;
+			BitDepth = OCIO_NAMESPACE::BIT_DEPTH_F16;
+			break;
+		case ERawImageFormat::RGBA32F:
+			Ordering = OCIO_NAMESPACE::CHANNEL_ORDERING_RGBA;
+			BitDepth = OCIO_NAMESPACE::BIT_DEPTH_F32;
+			break;
+		default:
+
+			UE_LOG(LogOpenColorIOWrapper, Log, TEXT("Unsupported texture format."));
+			return nullptr;
+		}
+
+		return MakeUnique<OCIO_NAMESPACE::PackedImageDesc>(
+			InImage.RawData,
+			static_cast<long>(InImage.GetWidth()),
+			static_cast<long>(InImage.GetHeight()),
+			Ordering,
+			BitDepth,
+			OCIO_NAMESPACE::AutoStride,
+			OCIO_NAMESPACE::AutoStride,
+			OCIO_NAMESPACE::AutoStride
+			);
+	};
+}
+#endif
+
+TUniquePtr<FOpenColorIOConfigWrapper> FOpenColorIOConfigWrapper::CreateWorkingColorSpaceToInterchangeConfig()
+{
+#if WITH_OCIO
+	TRACE_CPUPROFILER_EVENT_SCOPE(FOpenColorIOConfigWrapper::CreateWorkingColorSpaceToInterchangeConfig)
+	using namespace OCIO_NAMESPACE;
+	using namespace UE::Color;
+
+	TUniquePtr<FOpenColorIOConfigWrapper> Result = MakeUnique<FOpenColorIOConfigWrapper>();
+
+	ColorSpaceRcPtr AP0 = ColorSpace::Create();
+	AP0->setName("ACES2065-1");
+	AP0->setBitDepth(BIT_DEPTH_F32);
+	AP0->setEncoding("scene-linear");
+
+	ColorSpaceRcPtr WCS = ColorSpace::Create();
+	WCS->setName(OpenColorIOWrapper::WorkingColorSpaceName);
+	WCS->setBitDepth(BIT_DEPTH_F32);
+	WCS->setEncoding("scene-linear");
+
+	const FMatrix44d TransformMat = Transpose<double>(FColorSpaceTransform(FColorSpace::GetWorking(), FColorSpace(EColorSpace::ACESAP0)));
+	MatrixTransformRcPtr MatrixTransform = MatrixTransform::Create();
+	MatrixTransform->setMatrix(&TransformMat.M[0][0]);
+	WCS->setTransform(MatrixTransform, COLORSPACE_DIR_TO_REFERENCE);
+
+	ConfigRcPtr InterchangeConfig = Config::Create();
+	InterchangeConfig->addColorSpace(AP0);
+	InterchangeConfig->addColorSpace(WCS);
+	InterchangeConfig->setRole("aces_interchange", "ACES2065-1");
+
+	Result->Pimpl->Config = InterchangeConfig;
+
+	return Result;
+#else
+	return nullptr;
+#endif
+}
+
+FOpenColorIOConfigWrapper::FOpenColorIOConfigWrapper()
+	: Pimpl(MakePimpl<FOpenColorIOConfigPimpl>())
+{ }
+
+FOpenColorIOConfigWrapper::FOpenColorIOConfigWrapper(FStringView FilePath, FInitializationOptions InOptions)
+	: FOpenColorIOConfigWrapper()
+{
+#if WITH_OCIO
+	TRACE_CPUPROFILER_EVENT_SCOPE(FOpenColorIOConfigWrapper::FOpenColorIOConfigWrapper)
+#if !PLATFORM_EXCEPTIONS_DISABLED
+	try
+#endif
+	{
+		using namespace OCIO_NAMESPACE;
+
+		ConstConfigRcPtr NewConfig = Config::CreateFromFile(StringCast<ANSICHAR>(FilePath.GetData()).Get());
+
+		if (NewConfig && InOptions.bAddWorkingColorSpace)
+		{
+			ConstColorSpaceRcPtr InterchangeCS = NewConfig->getColorSpace(NewConfig->getCanonicalName(OpenColorIOWrapper::InterchangeName));
+
+			// When the aces interchange color space is present, we add the working color space as an additional option.
+			if (InterchangeCS != nullptr && NewConfig->getColorSpace(OpenColorIOWrapper::WorkingColorSpaceName) == nullptr)
+			{
+				ColorSpaceRcPtr WorkingCS = InterchangeCS->createEditableCopy();
+				WorkingCS->setName(OpenColorIOWrapper::WorkingColorSpaceName);
+				WorkingCS->setFamily("UE");
+				WorkingCS->clearAliases();
+
+				ConfigRcPtr NewConfigCopy = NewConfig->createEditableCopy();
+				NewConfigCopy->addColorSpace(WorkingCS);
+
+				NewConfig = NewConfigCopy;
+			}
+		}
+
+		Pimpl->Config = NewConfig;
+	}
+#if !PLATFORM_EXCEPTIONS_DISABLED
+	catch (OCIO_NAMESPACE::Exception& Exc)
+	{
+		UE_LOG(LogOpenColorIOWrapper, Error, TEXT("Could not create OCIO configuration file for %s. Error message: %s."), FilePath.GetData(), StringCast<TCHAR>(Exc.what()).Get());
+	}
+#endif
+#else
+	UE_LOG(LogOpenColorIOWrapper, Error, TEXT("The OpenColorIO library is not available, functionality will be disabled."));
+#endif // WITH_OCIO
+}
+
+bool FOpenColorIOConfigWrapper::IsValid() const
+{
+#if WITH_OCIO
+	return Pimpl->Config != nullptr;
+#else
+	return false;
+#endif
+}
+
+int32 FOpenColorIOConfigWrapper::GetNumColorSpaces() const
+{
+#if WITH_OCIO
+	if (IsValid())
+	{
+		return Pimpl->Config->getNumColorSpaces(OCIO_NAMESPACE::SEARCH_REFERENCE_SPACE_ALL, OCIO_NAMESPACE::COLORSPACE_ACTIVE);
+	}
+#endif
+
+	return 0;
+}
+
+FString FOpenColorIOConfigWrapper::GetColorSpaceName(int32 Index) const
+{
+#if WITH_OCIO
+	if (IsValid())
+	{
+		const char* ColorSpaceName = Pimpl->Config->getColorSpaceNameByIndex(OCIO_NAMESPACE::SEARCH_REFERENCE_SPACE_ALL, OCIO_NAMESPACE::COLORSPACE_ACTIVE, Index);
+
+		return StringCast<TCHAR>(ColorSpaceName).Get();
+	}
+#endif
+
+	return {};
+}
+
+int32 FOpenColorIOConfigWrapper::GetColorSpaceIndex(const TCHAR* InColorSpaceName)
+{
+#if WITH_OCIO
+	if (IsValid())
+	{
+		return Pimpl->Config->getIndexForColorSpace(StringCast<ANSICHAR>(InColorSpaceName).Get());
+	}
+#endif
+
+	return false;
+}
+
+FString FOpenColorIOConfigWrapper::GetColorSpaceFamilyName(const TCHAR* InColorSpaceName) const
+{
+#if WITH_OCIO
+	if (IsValid())
+	{
+		OCIO_NAMESPACE::ConstColorSpaceRcPtr ColorSpace = Pimpl->Config->getColorSpace(StringCast<ANSICHAR>(InColorSpaceName).Get());
+		
+		if (ColorSpace != nullptr)
+		{
+			return StringCast<TCHAR>(ColorSpace->getFamily()).Get();
+		}
+	}
+#endif
+
+	return {};
+}
+
+int32 FOpenColorIOConfigWrapper::GetNumDisplays() const
+{
+#if WITH_OCIO
+	if (IsValid())
+	{
+		return Pimpl->Config->getNumDisplays();
+	}
+#endif
+
+	return 0;
+}
+
+FString FOpenColorIOConfigWrapper::GetDisplayName(int32 Index) const
+{
+#if WITH_OCIO
+	if (IsValid())
+	{
+		const char* DisplayName = Pimpl->Config->getDisplay(Index);
+
+		return StringCast<TCHAR>(DisplayName).Get();
+	}
+#endif
+
+	return {};
+}
+
+int32 FOpenColorIOConfigWrapper::GetNumViews(const TCHAR* InDisplayName) const
+{
+#if WITH_OCIO
+	if (IsValid())
+	{
+		return Pimpl->Config->getNumViews(StringCast<ANSICHAR>(InDisplayName).Get());
+	}
+#endif
+
+	return 0;
+}
+
+FString FOpenColorIOConfigWrapper::GetViewName(const TCHAR* InDisplayName, int32 Index) const
+{
+#if WITH_OCIO
+	if (IsValid())
+	{
+		const char* ViewName = Pimpl->Config->getView(StringCast<ANSICHAR>(InDisplayName).Get(), Index);
+
+		return StringCast<TCHAR>(ViewName).Get();
+	}
+#endif
+
+	return {};
+}
+
+FString FOpenColorIOConfigWrapper::GetDisplayViewTransformName(const TCHAR* InDisplayName, const TCHAR* InViewName) const
+{
+#if WITH_OCIO
+	if (IsValid())
+	{
+		const char* TransformName = Pimpl->Config->getDisplayViewTransformName(StringCast<ANSICHAR>(InDisplayName).Get(), StringCast<ANSICHAR>(InViewName).Get());
+
+		return StringCast<TCHAR>(TransformName).Get();
+	}
+#endif
+
+	return {};
+}
+
+FString FOpenColorIOConfigWrapper::GetCacheID() const
+{
+#if WITH_OCIO
+	if (IsValid())
+	{
+		return StringCast<TCHAR>(Pimpl->Config->getCacheID()).Get();
+	}
+#endif
+
+	return {};
+}
+
+
+FOpenColorIOProcessorWrapper::FOpenColorIOProcessorWrapper()
+	: Pimpl(MakePimpl<FOpenColorIOProcessorPimpl>())
+	, OwnerConfig(nullptr)
+	, WorkingColorSpaceTransformType(EOpenColorIOWorkingColorSpaceTransform::None)
+{ }
+
+FOpenColorIOProcessorWrapper::FOpenColorIOProcessorWrapper(
+	const FOpenColorIOConfigWrapper* InConfig,
+	FStringView SourceColorSpace,
+	FStringView DestinationColorSpace,
+	EOpenColorIOWorkingColorSpaceTransform InWorkingColorSpaceTransformType,
+	const TMap<FString, FString>& ContextKeyValues)
+	: Pimpl(MakePimpl<FOpenColorIOProcessorPimpl>())
+	, OwnerConfig(InConfig)
+	, WorkingColorSpaceTransformType(InWorkingColorSpaceTransformType)
+{
+	ensure(OwnerConfig);
+
+#if WITH_OCIO
+	if (OwnerConfig && OwnerConfig->IsValid())
+	{
+		const OCIO_NAMESPACE::ConstConfigRcPtr& Config = OwnerConfig->Pimpl->Config;
+		 OCIO_NAMESPACE::ContextRcPtr Context = Config->getCurrentContext()->createEditableCopy();
+
+		for (const TPair<FString, FString>& KeyValue : ContextKeyValues)
+		{
+			Context->setStringVar(TCHAR_TO_ANSI(*KeyValue.Key), TCHAR_TO_ANSI(*KeyValue.Value));
+		}
+
+		Pimpl->Processor = Config->getProcessor(
+			Context,
+			StringCast<ANSICHAR>(SourceColorSpace.GetData()).Get(),
+			StringCast<ANSICHAR>(DestinationColorSpace.GetData()).Get()
+		);
+	}
+#endif
+}
+
+FOpenColorIOProcessorWrapper::FOpenColorIOProcessorWrapper(
+	const FOpenColorIOConfigWrapper* InConfig,
+	FStringView SourceColorSpace,
+	FStringView Display,
+	FStringView View,
+	bool bInverseDirection,
+	EOpenColorIOWorkingColorSpaceTransform InWorkingColorSpaceTransformType,
+	const TMap<FString, FString>& ContextKeyValues)
+	: Pimpl(MakePimpl<FOpenColorIOProcessorPimpl>())
+	, OwnerConfig(InConfig)
+	, WorkingColorSpaceTransformType(InWorkingColorSpaceTransformType)
+{
+	ensure(OwnerConfig);
+
+#if WITH_OCIO
+	if (OwnerConfig && OwnerConfig->IsValid())
+	{
+		const OCIO_NAMESPACE::ConstConfigRcPtr& Config = OwnerConfig->Pimpl->Config;
+		OCIO_NAMESPACE::ContextRcPtr Context = Config->getCurrentContext()->createEditableCopy();
+
+		for (const TPair<FString, FString>& KeyValue : ContextKeyValues)
+		{
+			Context->setStringVar(TCHAR_TO_ANSI(*KeyValue.Key), TCHAR_TO_ANSI(*KeyValue.Value));
+		}
+
+		Pimpl->Processor = Config->getProcessor(
+			StringCast<ANSICHAR>(SourceColorSpace.GetData()).Get(),
+			StringCast<ANSICHAR>(Display.GetData()).Get(),
+			StringCast<ANSICHAR>(View.GetData()).Get(),
+			static_cast<OCIO_NAMESPACE::TransformDirection>(bInverseDirection)
+		);
+	}
+#endif
+}
+
+bool FOpenColorIOProcessorWrapper::IsValid() const
+{
+#if WITH_OCIO
+	return Pimpl->Processor != nullptr;
+#else
+	return false;
+#endif
+}
+
+FOpenColorIOCPUProcessorWrapper::FOpenColorIOCPUProcessorWrapper(const FOpenColorIOProcessorWrapper& InOwner)
+	: OwnerProcessor(InOwner)
+{
+}
+
+bool FOpenColorIOCPUProcessorWrapper::IsValid() const
+{
+	return OwnerProcessor.IsValid();
+}
+
+bool FOpenColorIOCPUProcessorWrapper::TransformImage(const FImageView& InOutImage) const
+{
+#if WITH_OCIO
+	TRACE_CPUPROFILER_EVENT_SCOPE(FOpenColorIOCPUProcessorWrapper::TransformImage)
+#if !PLATFORM_EXCEPTIONS_DISABLED
+	try
+#endif
+	{
+		using namespace OCIO_NAMESPACE;
+
+		if (OwnerProcessor.IsValid())
+		{
+			TUniquePtr<PackedImageDesc> ImageDesc = GetImageDesc(InOutImage);
+			if (ImageDesc)
+			{
+				ConstConfigRcPtr	InterchangeConfig = IOpenColorIOWrapperModule::Get().GetWorkingColorSpaceToInterchangeConfig()->Pimpl->Config;
+				ConstConfigRcPtr Config = OwnerProcessor.OwnerConfig->Pimpl->Config;
+
+				BitDepth BitDepth = ImageDesc->getBitDepth();
+
+				// Conditionally apply a conversion from the working color space to interchange space
+				if (OwnerProcessor.WorkingColorSpaceTransformType == EOpenColorIOWorkingColorSpaceTransform::Source)
+				{
+					ConstProcessorRcPtr	InterchangeProcessor = InterchangeConfig->GetProcessorFromConfigs(
+						InterchangeConfig,
+						StringCast<ANSICHAR>(OpenColorIOWrapper::WorkingColorSpaceName).Get(),
+						Config,
+						Config->getCanonicalName(OpenColorIOWrapper::InterchangeName));
+
+					ConstCPUProcessorRcPtr InterchangeCPUProcessor = InterchangeProcessor->getOptimizedCPUProcessor(BitDepth, BitDepth, OPTIMIZATION_DEFAULT);
+					InterchangeCPUProcessor->apply(*ImageDesc);
+				}
+
+				// Apply the main color transformation
+				ConstCPUProcessorRcPtr CPUProcessor = OwnerProcessor.Pimpl->Processor->getOptimizedCPUProcessor(BitDepth, BitDepth, OPTIMIZATION_DEFAULT);
+				CPUProcessor->apply(*ImageDesc);
+
+				// Conditionally apply a conversion from the interchange space to the working color space
+				if (OwnerProcessor.WorkingColorSpaceTransformType == EOpenColorIOWorkingColorSpaceTransform::Destination)
+				{
+					ConstProcessorRcPtr	InterchangeProcessor = InterchangeConfig->GetProcessorFromConfigs(
+						Config,
+						Config->getCanonicalName(OpenColorIOWrapper::InterchangeName),
+						InterchangeConfig,
+						StringCast<ANSICHAR>(OpenColorIOWrapper::WorkingColorSpaceName).Get());
+
+					ConstCPUProcessorRcPtr InterchangeCPUProcessor = InterchangeProcessor->getOptimizedCPUProcessor(BitDepth, BitDepth, OPTIMIZATION_DEFAULT);
+					InterchangeCPUProcessor->apply(*ImageDesc);
+				}
+
+				return true;
+			}
+		}
+	}
+#if !PLATFORM_EXCEPTIONS_DISABLED
+	catch (OCIO_NAMESPACE::Exception& Exc)
+	{
+		UE_LOG(LogOpenColorIOWrapper, Log, TEXT("Failed to transform image. Error message: %s"), StringCast<TCHAR>(Exc.what()).Get());
+	}
+#endif
+#endif // WITH_OCIO
+
+	return false;
+}
+
+bool FOpenColorIOCPUProcessorWrapper::TransformImage(const FImageView& SrcImage, const FImageView& DestImage) const
+{
+#if WITH_OCIO
+	TRACE_CPUPROFILER_EVENT_SCOPE(FOpenColorIOCPUProcessorWrapper::TransformImage)
+#if !PLATFORM_EXCEPTIONS_DISABLED
+	try
+#endif
+	{
+		using namespace OCIO_NAMESPACE;
+
+		if (OwnerProcessor.IsValid())
+		{
+			TUniquePtr<PackedImageDesc> SrcImageDesc = GetImageDesc(SrcImage);
+			TUniquePtr<PackedImageDesc> DestImageDesc = GetImageDesc(DestImage);
+			if (SrcImageDesc && DestImageDesc)
+			{
+				ConstConfigRcPtr	InterchangeConfig = IOpenColorIOWrapperModule::Get().GetWorkingColorSpaceToInterchangeConfig()->Pimpl->Config;
+				ConstConfigRcPtr	Config = OwnerProcessor.OwnerConfig->Pimpl->Config;
+
+				BitDepth SrcBitDepth = SrcImageDesc->getBitDepth();
+				BitDepth DestBitDepth = DestImageDesc->getBitDepth();
+
+				// Conditionally apply a conversion from the working color space to interchange space
+				if (OwnerProcessor.WorkingColorSpaceTransformType == EOpenColorIOWorkingColorSpaceTransform::Source)
+				{
+					ConstProcessorRcPtr	InterchangeProcessor = InterchangeConfig->GetProcessorFromConfigs(
+						InterchangeConfig,
+						StringCast<ANSICHAR>(OpenColorIOWrapper::WorkingColorSpaceName).Get(),
+						Config,
+						Config->getCanonicalName(OpenColorIOWrapper::InterchangeName));
+
+					ConstCPUProcessorRcPtr InterchangeCPUProcessor = InterchangeProcessor->getOptimizedCPUProcessor(SrcBitDepth, SrcBitDepth, OPTIMIZATION_DEFAULT);
+					InterchangeCPUProcessor->apply(*SrcImageDesc);
+				}
+
+
+				// Apply the main color transformation
+				ConstCPUProcessorRcPtr CPUProcessor = OwnerProcessor.Pimpl->Processor->getOptimizedCPUProcessor(SrcBitDepth, DestBitDepth, OPTIMIZATION_DEFAULT);
+				CPUProcessor->apply(*SrcImageDesc, *DestImageDesc);
+
+				// Conditionally apply a conversion from the interchange space to the working color space
+				if (OwnerProcessor.WorkingColorSpaceTransformType == EOpenColorIOWorkingColorSpaceTransform::Destination)
+				{
+					ConstProcessorRcPtr	InterchangeProcessor = InterchangeConfig->GetProcessorFromConfigs(
+						Config,
+						Config->getCanonicalName(OpenColorIOWrapper::InterchangeName),
+						InterchangeConfig,
+						StringCast<ANSICHAR>(OpenColorIOWrapper::WorkingColorSpaceName).Get());
+
+					ConstCPUProcessorRcPtr InterchangeCPUProcessor = InterchangeProcessor->getOptimizedCPUProcessor(DestBitDepth, DestBitDepth, OPTIMIZATION_DEFAULT);
+					InterchangeCPUProcessor->apply(*DestImageDesc);
+				}
+
+				return true;
+			}
+		}
+	}
+#if !PLATFORM_EXCEPTIONS_DISABLED
+	catch (OCIO_NAMESPACE::Exception& exception)
+	{
+		UE_LOG(LogOpenColorIOWrapper, Log, TEXT("Failed to transform image. Error message: %s"), StringCast<TCHAR>(exception.what()).Get());
+	}
+#endif
+#endif // WITH_OCIO
+
+	return false;
+}
+
+
+
+FOpenColorIOGPUProcessorWrapper::FOpenColorIOGPUProcessorWrapper(const FOpenColorIOProcessorWrapper& InOwner, FInitializationOptions InShaderParams)
+	: OwnerProcessor(InOwner)
+	, GPUPimpl(MakePimpl<FOpenColorIOGPUProcessorPimpl>())
+{
+#if WITH_OCIO
+	TRACE_CPUPROFILER_EVENT_SCOPE(FOpenColorIOGPUProcessorWrapper::FOpenColorIOGPUProcessorWrapper)
+#if !PLATFORM_EXCEPTIONS_DISABLED
+	try
+#endif
+	{
+		using namespace OCIO_NAMESPACE;
+
+		GpuShaderDescRcPtr ShaderDescription = GpuShaderDesc::CreateShaderDesc();
+		ShaderDescription->setLanguage(GPU_LANGUAGE_HLSL_DX11);
+		ShaderDescription->setFunctionName(StringCast<ANSICHAR>(OpenColorIOWrapper::ShaderFunctionName).Get());
+		ShaderDescription->setResourcePrefix("Ocio");
+
+		ConstGPUProcessorRcPtr GPUProcessor = nullptr;
+		OptimizationFlags OptFlags = FOpenColorIOProcessorPimpl::GetOptimizationFlags();
+
+		if (InShaderParams.bIsLegacy)
+		{
+			unsigned int EdgeLength = static_cast<unsigned int>(OpenColorIOWrapper::Legacy3dEdgeLength);
+			GPUProcessor = OwnerProcessor.Pimpl->Processor->getOptimizedLegacyGPUProcessor(OptFlags, EdgeLength);
+		}
+		else
+		{
+			GPUProcessor = OwnerProcessor.Pimpl->Processor->getOptimizedGPUProcessor(FOpenColorIOProcessorPimpl::GetOptimizationFlags());
+		}
+		GPUProcessor->extractGpuShaderInfo(ShaderDescription);
+
+		GPUPimpl->Processor = GPUProcessor;
+		GPUPimpl->ShaderDescription = ShaderDescription;
+	}
+#if !PLATFORM_EXCEPTIONS_DISABLED
+	catch (OCIO_NAMESPACE::Exception& Exc)
+	{
+		//UE_LOG(LogOpenColorIOWrapper, Error, TEXT("Failed to fetch shader info for color transform %s. Error message: %s"), *GetTransformFriendlyName(), StringCast<TCHAR>(Exc.what()).Get());
+		UE_LOG(LogOpenColorIOWrapper, Error, TEXT("Failed to fetch shader info for color transform. Error message: %s"), StringCast<TCHAR>(Exc.what()).Get());
+	}
+#endif
+#endif //WITH_OCIO
+}
+
+bool FOpenColorIOGPUProcessorWrapper::IsValid() const
+{
+#if WITH_OCIO
+	return OwnerProcessor.IsValid() && GPUPimpl->Processor != nullptr && GPUPimpl->ShaderDescription != nullptr;
+#else
+	return false;
+#endif
+}
+
+bool FOpenColorIOGPUProcessorWrapper::GetShader(FString& OutShaderCode, FString& OutShaderCacheID) const
+{
+#if WITH_OCIO
+	if (IsValid())
+	{
+		ensureMsgf(GPUPimpl->ShaderDescription->getNumDynamicProperties() == 0, TEXT("We do not currently support dynamic properties."));
+
+		OutShaderCode = StringCast<TCHAR>(GPUPimpl->ShaderDescription->getShaderText()).Get();
+		OutShaderCacheID = StringCast<TCHAR>(GPUPimpl->ShaderDescription->getCacheID()).Get();
+
+		return true;
+	}
+#endif
+
+	return false;
+}
+
+uint32 FOpenColorIOGPUProcessorWrapper::GetNum3DTextures() const
+{
+#if WITH_OCIO
+	if (IsValid())
+	{
+		return GPUPimpl->ShaderDescription->getNum3DTextures();
+	}
+#endif
+
+	return 0;
+}
+
+bool FOpenColorIOGPUProcessorWrapper::Get3DTexture(uint32 InIndex, FName& OutName, uint32& OutEdgeLength, TextureFilter& OutTextureFilter, const float*& OutData) const
+{
+#if WITH_OCIO
+	ensure(InIndex < GetNum3DTextures());
+
+	if (IsValid())
+	{
+		const ANSICHAR* TextureName = nullptr;
+		const ANSICHAR* SamplerName = nullptr;
+		OCIO_NAMESPACE::Interpolation Interpolation = OCIO_NAMESPACE::INTERP_TETRAHEDRAL;
+
+		// Read texture information
+		GPUPimpl->ShaderDescription->get3DTexture(InIndex, TextureName, SamplerName, OutEdgeLength, Interpolation);
+
+		// Read texture data
+		OutData = 0x0;
+		GPUPimpl->ShaderDescription->get3DTextureValues(InIndex, OutData);
+
+		OutName = FName(TextureName);
+
+		OutTextureFilter = TF_Bilinear;
+		if (Interpolation == OCIO_NAMESPACE::Interpolation::INTERP_NEAREST || Interpolation == OCIO_NAMESPACE::Interpolation::INTERP_TETRAHEDRAL)
+		{
+			OutTextureFilter = TF_Nearest;
+		}
+
+		return TextureName && OutEdgeLength > 0 && OutData;
+	}
+#endif
+
+	return false;
+}
+uint32 FOpenColorIOGPUProcessorWrapper::GetNumTextures() const
+{
+#if WITH_OCIO
+	if (IsValid())
+	{
+		return GPUPimpl->ShaderDescription->getNumTextures();
+	}
+#endif
+
+	return 0;
+}
+
+bool FOpenColorIOGPUProcessorWrapper::GetTexture(uint32 InIndex, FName& OutName, uint32& OutWidth, uint32& OutHeight, TextureFilter& OutTextureFilter, bool& bOutRedChannelOnly, const float*& OutData) const
+{
+	ensure(InIndex < GetNumTextures());
+
+#if WITH_OCIO
+	if (IsValid())
+	{
+		const ANSICHAR* TextureName = nullptr;
+		const ANSICHAR* SamplerName = nullptr;
+		OCIO_NAMESPACE::GpuShaderDesc::TextureType Channel = OCIO_NAMESPACE::GpuShaderDesc::TEXTURE_RGB_CHANNEL;
+		OCIO_NAMESPACE::Interpolation Interpolation = OCIO_NAMESPACE::Interpolation::INTERP_LINEAR;
+
+		// Read texture information
+		GPUPimpl->ShaderDescription->getTexture(InIndex, TextureName, SamplerName, OutWidth, OutHeight, Channel, Interpolation);
+
+		// Read texture data
+		OutData = 0x0;
+		GPUPimpl->ShaderDescription->getTextureValues(InIndex, OutData);
+
+		OutName = FName(TextureName);
+		OutTextureFilter = Interpolation == OCIO_NAMESPACE::Interpolation::INTERP_NEAREST ? TF_Nearest : TF_Bilinear;
+		bOutRedChannelOnly = Channel == OCIO_NAMESPACE::GpuShaderCreator::TEXTURE_RED_CHANNEL;
+
+		return TextureName && OutWidth > 0 && OutHeight > 0 && OutData;
+	}
+#endif
+
+	return false;
+}
+
+FString FOpenColorIOGPUProcessorWrapper::GetCacheID() const
+{
+#if WITH_OCIO
+	if (IsValid())
+	{
+		return StringCast<TCHAR>(GPUPimpl->Processor->getCacheID()).Get();
+	}
+#endif
+
+	return {};
+}
+
+#undef LOCTEXT_NAMESPACE
diff --git a/Engine/Plugins/Compositing/OpenColorIO/Source/OpenColorIOWrapper/Private/OpenColorIOWrapperModule.cpp b/Engine/Plugins/Compositing/OpenColorIO/Source/OpenColorIOWrapper/Private/OpenColorIOWrapperModule.cpp
new file mode 100644
--- /dev/null
+++ b/Engine/Plugins/Compositing/OpenColorIO/Source/OpenColorIOWrapper/Private/OpenColorIOWrapperModule.cpp
@@ -0,0 +1,50 @@
+// Copyright Epic Games, Inc. All Rights Reserved.
+
+#include "OpenColorIOWrapperModule.h"
+
+#include "OpenColorIOWrapper.h"
+#include "OpenColorIOLibHandler.h"
+
+
+DEFINE_LOG_CATEGORY(LogOpenColorIOWrapper);
+
+#define OPENCOLORIOWRAPPER_MODULE_NAME "OpenColorIOWrapper"
+#define LOCTEXT_NAMESPACE "OpenColorIOWrapperModule"
+
+class FOpenColorIOWrapperModule : public IOpenColorIOWrapperModule
+{
+public:
+	//~ Begin IModuleInterface interface
+	virtual void StartupModule() override
+	{
+		bInitializedLib = FOpenColorIOLibHandler::Initialize();
+	}
+
+	virtual void ShutdownModule() override
+	{
+		FOpenColorIOLibHandler::Shutdown();
+		bInitializedLib = false;
+	}
+	//~ End IModuleInterface interface
+
+	virtual const FOpenColorIOConfigWrapper* GetWorkingColorSpaceToInterchangeConfig() override
+	{
+		check(bInitializedLib);
+
+		if (!InterchangeConfig)
+		{
+			InterchangeConfig = FOpenColorIOConfigWrapper::CreateWorkingColorSpaceToInterchangeConfig();
+		}
+
+		return InterchangeConfig.Get();
+	}
+private:
+
+	bool bInitializedLib = false;
+	
+	TUniquePtr<FOpenColorIOConfigWrapper> InterchangeConfig;
+};
+
+IMPLEMENT_MODULE(FOpenColorIOWrapperModule, OpenColorIO);
+
+#undef LOCTEXT_NAMESPACE
diff --git a/Engine/Plugins/Compositing/OpenColorIO/Source/OpenColorIOWrapper/Public/OpenColorIOWrapper.h b/Engine/Plugins/Compositing/OpenColorIO/Source/OpenColorIOWrapper/Public/OpenColorIOWrapper.h
new file mode 100644
--- /dev/null
+++ b/Engine/Plugins/Compositing/OpenColorIO/Source/OpenColorIOWrapper/Public/OpenColorIOWrapper.h
@@ -0,0 +1,261 @@
+// Copyright Epic Games, Inc. All Rights Reserved.
+
+#pragma once
+
+#include "CoreTypes.h"
+#include "Containers/ContainersFwd.h"
+#include "Containers/StringView.h"
+#include "Containers/UnrealString.h"
+#include "Containers/Map.h"
+#include "Templates/PimplPtr.h"
+#include "Templates/UniquePtr.h"
+
+#include "OpenColorIOWrapperDefines.h"
+
+struct FImageView;
+enum TextureFilter : int;
+
+namespace OpenColorIOWrapper
+{
+	/** OCIO_NAMESPACE::ROLE_INTERCHANGE_SCENE equivalent, since we currently cannot delay-load this definition. */
+	static constexpr ANSICHAR InterchangeName[] = "aces_interchange";
+	
+	/** Color space name of the engine's working color space we insert in OpenColorIO configs. */
+	static constexpr ANSICHAR WorkingColorSpaceName[] = "Working Color Space";
+
+	/** Default generated shader function name. */
+	static constexpr TCHAR ShaderFunctionName[] = TEXT("OCIOConvert");
+
+	/** Default LUT size used in the legacy gpu processor */
+	static constexpr uint32 Legacy3dEdgeLength = 65;
+
+	/** Get the OpenColorIO version string. */
+	OPENCOLORIOWRAPPER_API const TCHAR* GetVersion();
+}
+
+class OPENCOLORIOWRAPPER_API FOpenColorIOConfigWrapper final
+{
+public:
+
+	/** Config initialization options. */
+	struct FInitializationOptions
+	{
+		FInitializationOptions() : bAddWorkingColorSpace(false) { };
+
+		bool bAddWorkingColorSpace;
+	};
+
+	/** Constructor. */
+	FOpenColorIOConfigWrapper();
+
+	/**
+	* Constructor.
+	* 
+	* @param FilePath Config absolute file path.
+	* @param InOptions Initialization options.
+	*/
+	FOpenColorIOConfigWrapper(FStringView FilePath, FInitializationOptions InOptions);
+
+	/** Valid when the native config has been successfully created and isn't null. */
+	bool IsValid() const;
+	
+	/** Get the number of color spaces in the configuration. */
+	int32 GetNumColorSpaces() const;
+	
+	/** Get a color space name at an index. */
+	FString GetColorSpaceName(int32 Index) const;
+	
+	/** Get the index of a color space, -1 if none. */
+	int32 GetColorSpaceIndex(const TCHAR* InColorSpaceName);
+	
+	/** Get the family name for a color space. */
+	FString GetColorSpaceFamilyName(const TCHAR* InColorSpaceName) const;
+	
+	/** Get the number of displays in the configuration. */
+	int32 GetNumDisplays() const;
+	
+	/** Get a display name at an index. */
+	FString GetDisplayName(int32 Index) const;
+	
+	/** Get the number of views for a display. */
+	int32 GetNumViews(const TCHAR* InDisplayName) const;
+	
+	/** Get a view name for its display and index. */
+	FString GetViewName(const TCHAR* InDisplayName, int32 Index) const;
+
+	/** Get a display-view trnasform name. */
+	FString GetDisplayViewTransformName(const TCHAR* InDisplayName, const TCHAR* InViewName) const;
+
+	/** Get the string hash of the config. */
+	FString GetCacheID() const;
+
+private:
+
+	/** Convenience to create a config between the working color space and the default interchange one. */
+	static TUniquePtr<FOpenColorIOConfigWrapper> CreateWorkingColorSpaceToInterchangeConfig();
+
+	TPimplPtr<struct FOpenColorIOConfigPimpl> Pimpl;
+
+	friend class FOpenColorIOWrapperModule;
+	friend class FOpenColorIOProcessorWrapper;
+	friend class FOpenColorIOCPUProcessorWrapper;
+	friend class FOpenColorIOGPUProcessorWrapper;
+};
+
+class OPENCOLORIOWRAPPER_API FOpenColorIOProcessorWrapper final
+{
+public:
+	/** Constructor. */
+	FOpenColorIOProcessorWrapper();
+	
+	/**
+	* Constructor.
+	*
+	* @param InConfig Owner config.
+	* @param SourceColorSpace Source color space name.
+	* @param DestinationColorSpace Destination color space name.
+	* @param WorkingColorSpaceTransformType (Optional) Type of additional conversion from/to the working color space.
+	* @param ContextKeyValues (Optional) ContextKeyValues Additional context modifiers.
+	*/
+	FOpenColorIOProcessorWrapper(
+		const FOpenColorIOConfigWrapper* InConfig,
+		FStringView SourceColorSpace,
+		FStringView DestinationColorSpace,
+		EOpenColorIOWorkingColorSpaceTransform WorkingColorSpaceTransformType = EOpenColorIOWorkingColorSpaceTransform::None,
+		const TMap<FString, FString>& ContextKeyValues = {});
+
+	/**
+	* Constructor.
+	*
+	* @param InConfig Owner config.
+	* @param SourceColorSpace Source color space name.
+	* @param Display Display name in display-view transform.
+	* @param View View name in display-view transform.
+	* @param bInverseDirection Flag for inverse transform direction.
+	* @param WorkingColorSpaceTransformType (Optional) Type of additional conversion from/to the working color space.
+	* @param ContextKeyValues (Optional) ContextKeyValues Additional context modifiers.
+	*/
+	FOpenColorIOProcessorWrapper(
+		const FOpenColorIOConfigWrapper* InConfig,
+		FStringView SourceColorSpace,
+		FStringView Display,
+		FStringView View,
+		bool bInverseDirection = false,
+		EOpenColorIOWorkingColorSpaceTransform WorkingColorSpaceTransformType = EOpenColorIOWorkingColorSpaceTransform::None,
+		const TMap<FString, FString>& ContextKeyValues = {});
+
+	/** Valid when the processor has been successfully created and isn't null. */
+	bool IsValid() const;
+
+private:
+
+	TPimplPtr<struct FOpenColorIOProcessorPimpl> Pimpl;
+
+	const FOpenColorIOConfigWrapper* OwnerConfig;
+
+	EOpenColorIOWorkingColorSpaceTransform WorkingColorSpaceTransformType;
+
+	friend class FOpenColorIOCPUProcessorWrapper;
+	friend class FOpenColorIOGPUProcessorWrapper;
+};
+
+
+class OPENCOLORIOWRAPPER_API FOpenColorIOCPUProcessorWrapper final
+{
+public:
+	/**
+	* Constructor.
+	*
+	* @param InOwner Owner processor.
+	*/
+	FOpenColorIOCPUProcessorWrapper(const FOpenColorIOProcessorWrapper& InOwner);
+
+	/** Valid when the processor has been successfully created and isn't null. */
+	bool IsValid() const;
+
+	/** Apply the color transform in-place to the specified image. */
+	bool TransformImage(const FImageView& InOutImage) const;
+
+	/** Apply the color transform from the source image to the destination image. (The destination FImageView is const but what it points at is not.) */
+	bool TransformImage(const FImageView& SrcImage, const FImageView& DestImage) const;
+
+private:
+
+	const FOpenColorIOProcessorWrapper& OwnerProcessor;
+};
+
+class OPENCOLORIOWRAPPER_API FOpenColorIOGPUProcessorWrapper final
+{
+public:
+	/** Gpu processor initialization options. */
+	struct FInitializationOptions
+	{
+		FInitializationOptions() : bIsLegacy(false) { }
+
+		bool bIsLegacy;
+	};
+
+	/**
+	* Constructor.
+	*
+	* @param InOwner Owner processor.
+	* @param InOptions Initialization options.
+	*/
+	FOpenColorIOGPUProcessorWrapper(
+		const FOpenColorIOProcessorWrapper& InOwner,
+		FInitializationOptions InOptions = FInitializationOptions()
+	);
+
+	/** Valid when the processor has been successfully created and isn't null. */
+	bool IsValid() const;
+
+	/**
+	* Get the generated shader.
+	*
+	* @param OutShaderCode Generated shader text.
+	* @param OutShaderCacheID Generated shader hash string.
+	* @return True when the result is valid.
+	*/
+	bool GetShader(FString& OutShaderCode, FString& OutShaderCacheID) const;
+
+	/** Get the number of 3D LUT textures. */
+	uint32 GetNum3DTextures() const;
+
+	/**
+	* Get the index-specific 3D LUT texture used by the shader transform.
+	*
+	* @param InIndex Index of the texture.
+	* @param OutName Shader parameter name of the texture.
+	* @param OutEdgeLength Size of the texture.
+	* @param OutTextureFilter Texture filter type.
+	* @param OutData Raw texel data.
+	* @return True when the result is valid.
+	*/
+	bool Get3DTexture(uint32 InIndex, FName& OutName, uint32& OutEdgeLength, TextureFilter& OutTextureFilter, const float*& OutData) const;
+
+	/** Get the number of 2D LUT textures. (1D resources are disabled in our case.) */
+	uint32 GetNumTextures() const;
+
+	/**
+	* Get the index-specific 2D LUT texture used by the shader transform.
+	*
+	* @param InIndex Index of the texture.
+	* @param OutName Shader parameter name of the texture.
+	* @param OutWidth Width of the texture.
+	* @param OutHeight Height of the texture.
+	* @param OutTextureFilter Texture filter type.
+	* @param bOutRedChannelOnly Flag indicating whether the texture has a single channel or is RGB.
+	* @param OutData Raw texel data.
+	* @return True when the result is valid.
+	*/
+	bool GetTexture(uint32 InIndex, FName& OutName, uint32& OutWidth, uint32& OutHeight, TextureFilter& OutTextureFilter, bool& bOutRedChannelOnly, const float*& OutData) const;
+
+	/** Get the string hash of the gpu processor. */
+	FString GetCacheID() const;
+
+private:
+
+	const FOpenColorIOProcessorWrapper& OwnerProcessor;
+
+	TPimplPtr<struct FOpenColorIOGPUProcessorPimpl> GPUPimpl;
+};
diff --git a/Engine/Plugins/Compositing/OpenColorIO/Source/OpenColorIOWrapper/Public/OpenColorIOWrapperDefines.h b/Engine/Plugins/Compositing/OpenColorIO/Source/OpenColorIOWrapper/Public/OpenColorIOWrapperDefines.h
new file mode 100644
--- /dev/null
+++ b/Engine/Plugins/Compositing/OpenColorIO/Source/OpenColorIOWrapper/Public/OpenColorIOWrapperDefines.h
@@ -0,0 +1,14 @@
+// Copyright Epic Games, Inc. All Rights Reserved.
+
+#pragma once
+
+#include "HAL/Platform.h"
+
+/** Enum used to indicate whether the working color space should be used as a source or destination. */
+enum class EOpenColorIOWorkingColorSpaceTransform : uint8
+{
+	None = 0,
+	Source = 1,
+	Destination = 2
+};
+
diff --git a/Engine/Plugins/Compositing/OpenColorIO/Source/OpenColorIOWrapper/Public/OpenColorIOWrapperModule.h b/Engine/Plugins/Compositing/OpenColorIO/Source/OpenColorIOWrapper/Public/OpenColorIOWrapperModule.h
new file mode 100644
--- /dev/null
+++ b/Engine/Plugins/Compositing/OpenColorIO/Source/OpenColorIOWrapper/Public/OpenColorIOWrapperModule.h
@@ -0,0 +1,48 @@
+// Copyright Epic Games, Inc. All Rights Reserved.
+
+#pragma once
+
+#include "Modules/ModuleManager.h"
+#include "Logging/LogMacros.h"
+
+class FOpenColorIOConfigWrapper;
+
+DECLARE_LOG_CATEGORY_EXTERN(LogOpenColorIOWrapper, Log, All);
+
+/**
+ * Interface for the OpenColorIO module.
+ */
+class IOpenColorIOWrapperModule : public IModuleInterface
+{
+public:
+	/**
+	 * Singleton-like access to this module's interface.  This is just for convenience!
+	 * Beware of calling this during the shutdown phase, though.  Your module might have been unloaded already.
+	 *
+	 * @return Returns singleton instance, loading the module on demand if needed
+	 */
+	static IOpenColorIOWrapperModule& Get()
+	{
+		static const FName ModuleName = "OpenColorIOWrapper";
+		return FModuleManager::LoadModuleChecked<IOpenColorIOWrapperModule>(ModuleName);
+	}
+
+	/**
+	 * Checks to see if this module is loaded and ready.  It is only valid to call Get() during shutdown if IsAvailable() returns true.
+	 *
+	 * @return True if the module is loaded and ready to use
+	 */
+	static bool IsAvailable()
+	{
+		static const FName ModuleName = "OpenColorIOWrapper";
+		return FModuleManager::Get().IsModuleLoaded(ModuleName);
+	}
+
+	/**
+	 * Returns a minimal dynamically-created native config for conversions between interchange and working color spaces.
+	 */
+	virtual const FOpenColorIOConfigWrapper* GetWorkingColorSpaceToInterchangeConfig() = 0;
+
+	/** Virtual destructor */
+	virtual ~IOpenColorIOWrapperModule() = default;
+};
