OCIO: Add support for programmatically controlling the active viewport configuration.

diff --git a/Engine/Plugins/Compositing/OpenColorIO/Source/OpenColorIOEditor/Private/OpenColorIOEditorBlueprintLibrary.cpp b/Engine/Plugins/Compositing/OpenColorIO/Source/OpenColorIOEditor/Private/OpenColorIOEditorBlueprintLibrary.cpp
new file mode 100644
--- /dev/null
+++ b/Engine/Plugins/Compositing/OpenColorIO/Source/OpenColorIOEditor/Private/OpenColorIOEditorBlueprintLibrary.cpp
@@ -0,0 +1,18 @@
+// Copyright Epic Games, Inc. All Rights Reserved.
+
+#include "OpenColorIOEditorBlueprintLibrary.h"
+
+#include "Modules/ModuleManager.h"
+
+#include "OpenColorIOEditorModule.h"
+#include "OpenColorIOColorSpace.h"
+
+#include UE_INLINE_GENERATED_CPP_BY_NAME(OpenColorIOEditorBlueprintLibrary)
+
+
+void UOpenColorIOEditorBlueprintLibrary::SetActiveViewportConfiguration(const FOpenColorIODisplayConfiguration& InConfiguration)
+{
+	FOpenColorIOEditorModule& OpenColorIOEditorModule = FModuleManager::LoadModuleChecked<FOpenColorIOEditorModule>("OpenColorIOEditor");
+
+	OpenColorIOEditorModule.SetActiveViewportConfiguration(InConfiguration);
+}
diff --git a/Engine/Plugins/Compositing/OpenColorIO/Source/OpenColorIOEditor/Private/OpenColorIOEditorModule.cpp b/Engine/Plugins/Compositing/OpenColorIO/Source/OpenColorIOEditor/Private/OpenColorIOEditorModule.cpp
--- a/Engine/Plugins/Compositing/OpenColorIO/Source/OpenColorIOEditor/Private/OpenColorIOEditorModule.cpp
+++ b/Engine/Plugins/Compositing/OpenColorIO/Source/OpenColorIOEditor/Private/OpenColorIOEditorModule.cpp
@@ -52,6 +52,11 @@ void FOpenColorIOEditorModule::ShutdownModule()
 	FCoreDelegates::OnFEngineLoopInitComplete.RemoveAll(this);
 }
 
+void FOpenColorIOEditorModule::SetActiveViewportConfiguration(const FOpenColorIODisplayConfiguration& InConfiguration)
+{
+	OnDisplayConfigurationChanged(InConfiguration);
+}
+
 void FOpenColorIOEditorModule::RegisterCustomizations()
 {
 	FPropertyEditorModule& PropertyModule = FModuleManager::LoadModuleChecked<FPropertyEditorModule>("PropertyEditor");
diff --git a/Engine/Plugins/Compositing/OpenColorIO/Source/OpenColorIOEditor/Private/OpenColorIOEditorModule.h b/Engine/Plugins/Compositing/OpenColorIO/Source/OpenColorIOEditor/Private/OpenColorIOEditorModule.h
--- a/Engine/Plugins/Compositing/OpenColorIO/Source/OpenColorIOEditor/Private/OpenColorIOEditorModule.h
+++ b/Engine/Plugins/Compositing/OpenColorIO/Source/OpenColorIOEditor/Private/OpenColorIOEditorModule.h
@@ -29,6 +29,9 @@ public:
 	virtual void ShutdownModule() override;
 	//~ End IModuleInterface interface
 
+	/** Set a color transform display configuration on the active editor viewport. */
+	void SetActiveViewportConfiguration(const FOpenColorIODisplayConfiguration& InConfiguration);
+
 protected:
 	void RegisterCustomizations();
 	void UnregisterCustomizations();
diff --git a/Engine/Plugins/Compositing/OpenColorIO/Source/OpenColorIOEditor/Public/OpenColorIOEditorBlueprintLibrary.h b/Engine/Plugins/Compositing/OpenColorIO/Source/OpenColorIOEditor/Public/OpenColorIOEditorBlueprintLibrary.h
new file mode 100644
--- /dev/null
+++ b/Engine/Plugins/Compositing/OpenColorIO/Source/OpenColorIOEditor/Public/OpenColorIOEditorBlueprintLibrary.h
@@ -0,0 +1,24 @@
+// Copyright Epic Games, Inc. All Rights Reserved.
+
+#pragma once
+
+#include "Kismet/BlueprintFunctionLibrary.h"
+
+#include "OpenColorIOEditorBlueprintLibrary.generated.h"
+
+struct FOpenColorIODisplayConfiguration;
+
+UCLASS()
+class UOpenColorIOEditorBlueprintLibrary : public UBlueprintFunctionLibrary
+{
+	GENERATED_BODY()
+public:
+
+	/**
+	 * Set the active editor viewport's display configuration color transform .
+	 *
+	 * @param InDisplayConfiguration Display configuration color transform
+	 */
+	UFUNCTION(BlueprintCallable, Category = "OpenColorIO")
+	static OPENCOLORIOEDITOR_API void SetActiveViewportConfiguration(const FOpenColorIODisplayConfiguration& InConfiguration);
+};
