OCIO: Fix BP name spacing with meta display names.

#rnx

diff --git a/Engine/Plugins/Compositing/OpenColorIO/Source/OpenColorIO/Public/OpenColorIOColorSpace.h b/Engine/Plugins/Compositing/OpenColorIO/Source/OpenColorIO/Public/OpenColorIOColorSpace.h
--- a/Engine/Plugins/Compositing/OpenColorIO/Source/OpenColorIO/Public/OpenColorIOColorSpace.h
+++ b/Engine/Plugins/Compositing/OpenColorIO/Source/OpenColorIO/Public/OpenColorIOColorSpace.h
@@ -13,7 +13,7 @@ class UOpenColorIOConfiguration;
  * Structure to identify a ColorSpace as described in an OCIO configuration file. 
  * Members are populated by data coming from a config file.
  */
-USTRUCT(BlueprintType)
+USTRUCT(BlueprintType, meta = (DisplayName = "OpenColorIO Color Space"))
 struct OPENCOLORIO_API FOpenColorIOColorSpace
 {
 	GENERATED_BODY()
@@ -73,7 +73,7 @@ public:
 /**
  * Transformation direction type for display-view transformations.
  */
-UENUM(BlueprintType)
+UENUM(BlueprintType, meta = (DisplayName = "OpenColorIO View Transform Direction"))
 enum class EOpenColorIOViewTransformDirection : uint8
 {
 	Forward = 0     UMETA(DisplayName = "Forward"),
@@ -81,7 +81,7 @@ enum class EOpenColorIOViewTransformDirection : uint8
 };
 
 
-USTRUCT(BlueprintType)
+USTRUCT(BlueprintType, meta = (DisplayName = "OpenColorIO Display View"))
 struct OPENCOLORIO_API FOpenColorIODisplayView
 {
 	GENERATED_BODY()
@@ -117,7 +117,7 @@ public:
 /**
  * Identifies a OCIO ColorSpace conversion.
  */
-USTRUCT(BlueprintType)
+USTRUCT(BlueprintType, meta = (DisplayName = "OpenColorIO Color Conversion Settings"))
 struct OPENCOLORIO_API FOpenColorIOColorConversionSettings
 {
 	GENERATED_BODY()
@@ -192,7 +192,7 @@ private:
 /**
  * Identifies an OCIO Display look configuration 
  */
-USTRUCT(BlueprintType)
+USTRUCT(BlueprintType, meta = (DisplayName = "OpenColorIO Display Configuration"))
 struct OPENCOLORIO_API FOpenColorIODisplayConfiguration
 {
 	GENERATED_BODY()
diff --git a/Engine/Plugins/Compositing/OpenColorIO/Source/OpenColorIO/Public/OpenColorIOConfiguration.h b/Engine/Plugins/Compositing/OpenColorIO/Source/OpenColorIO/Public/OpenColorIOConfiguration.h
--- a/Engine/Plugins/Compositing/OpenColorIO/Source/OpenColorIO/Public/OpenColorIOConfiguration.h
+++ b/Engine/Plugins/Compositing/OpenColorIO/Source/OpenColorIO/Public/OpenColorIOConfiguration.h
@@ -26,7 +26,7 @@ namespace ERHIFeatureLevel { enum Type : int; }
 /**
  * Asset to manage allowed OpenColorIO color spaces. This will create required transform objects.
  */
-UCLASS(BlueprintType)
+UCLASS(BlueprintType, meta = (DisplayName = "OpenColorIO Configuration"))
 class OPENCOLORIO_API UOpenColorIOConfiguration : public UObject
 {
 	GENERATED_BODY()
diff --git a/Engine/Plugins/Compositing/OpenColorIO/Source/OpenColorIO/Public/OpenColorIODisplayExtensionWrapper.h b/Engine/Plugins/Compositing/OpenColorIO/Source/OpenColorIO/Public/OpenColorIODisplayExtensionWrapper.h
--- a/Engine/Plugins/Compositing/OpenColorIO/Source/OpenColorIO/Public/OpenColorIODisplayExtensionWrapper.h
+++ b/Engine/Plugins/Compositing/OpenColorIO/Source/OpenColorIO/Public/OpenColorIODisplayExtensionWrapper.h
@@ -11,7 +11,7 @@
  * This Blueprintable object can hold an OCIO Scene View Extension. 
  * You can change its OCIO config, and specify the context in which you want it to be active on.
  */
-UCLASS(Blueprintable)
+UCLASS(Blueprintable, meta = (DisplayName = "OpenColorIO Display Extension Wrapper"))
 class OPENCOLORIO_API UOpenColorIODisplayExtensionWrapper : public UObject
 {
 	GENERATED_BODY()
