OCIO: Fix duplicate module implementation.

diff --git a/Engine/Plugins/Compositing/OpenColorIO/Source/OpenColorIOWrapper/Private/OpenColorIOWrapperModule.cpp b/Engine/Plugins/Compositing/OpenColorIO/Source/OpenColorIOWrapper/Private/OpenColorIOWrapperModule.cpp
--- a/Engine/Plugins/Compositing/OpenColorIO/Source/OpenColorIOWrapper/Private/OpenColorIOWrapperModule.cpp
+++ b/Engine/Plugins/Compositing/OpenColorIO/Source/OpenColorIOWrapper/Private/OpenColorIOWrapperModule.cpp
@@ -45,6 +45,6 @@ private:
 	TUniquePtr<FOpenColorIOConfigWrapper> InterchangeConfig;
 };
 
-IMPLEMENT_MODULE(FOpenColorIOWrapperModule, OpenColorIO);
+IMPLEMENT_MODULE(FOpenColorIOWrapperModule, OpenColorIOWrapper);
 
 #undef LOCTEXT_NAMESPACE
