OCIO: Small improvements to new wrapper class.

Added missing exception handling. Using parent reference for consistency.

diff --git a/Engine/Plugins/Compositing/OpenColorIO/Source/OpenColorIO/Private/OpenColorIOColorTransform.cpp b/Engine/Plugins/Compositing/OpenColorIO/Source/OpenColorIO/Private/OpenColorIOColorTransform.cpp
--- a/Engine/Plugins/Compositing/OpenColorIO/Source/OpenColorIO/Private/OpenColorIOColorTransform.cpp
+++ b/Engine/Plugins/Compositing/OpenColorIO/Source/OpenColorIO/Private/OpenColorIOColorTransform.cpp
@@ -31,11 +31,14 @@ namespace
 		check(InTransform);
 		check(InTransform->ConfigurationOwner);
 
+		const FOpenColorIOConfigWrapper* ConfigWrapper = InTransform->ConfigurationOwner->GetConfigWrapper();
+		ensure(ConfigWrapper);
+
 		EOpenColorIOViewTransformDirection DisplayViewDirection;
 		if (InTransform->GetDisplayViewDirection(DisplayViewDirection))
 		{
 			return FOpenColorIOProcessorWrapper(
-				InTransform->ConfigurationOwner->GetConfigWrapper(),
+				*ConfigWrapper,
 				InTransform->SourceColorSpace,
 				InTransform->Display,
 				InTransform->View,
@@ -46,7 +49,7 @@ namespace
 		else
 		{
 			return FOpenColorIOProcessorWrapper(
-				InTransform->ConfigurationOwner->GetConfigWrapper(),
+				*ConfigWrapper,
 				InTransform->SourceColorSpace,
 				InTransform->DestinationColorSpace,
 				InWCSTrasnformType,
diff --git a/Engine/Plugins/Compositing/OpenColorIO/Source/OpenColorIOWrapper/Private/OpenColorIOWrapper.cpp b/Engine/Plugins/Compositing/OpenColorIO/Source/OpenColorIOWrapper/Private/OpenColorIOWrapper.cpp
--- a/Engine/Plugins/Compositing/OpenColorIO/Source/OpenColorIOWrapper/Private/OpenColorIOWrapper.cpp
+++ b/Engine/Plugins/Compositing/OpenColorIO/Source/OpenColorIOWrapper/Private/OpenColorIOWrapper.cpp
@@ -22,14 +22,14 @@ const TCHAR* OpenColorIOWrapper::GetVersion()
 	return TEXT(OCIO_VERSION);
 #else
 	return TEXT("");
-#endif
+#endif // WITH_OCIO
 }
 
 struct FOpenColorIOConfigPimpl
 {
 #if WITH_OCIO
 	OCIO_NAMESPACE::ConstConfigRcPtr Config = nullptr;
-#endif
+#endif // WITH_OCIO
 };
 
 struct FOpenColorIOProcessorPimpl
@@ -45,7 +45,7 @@ struct FOpenColorIOProcessorPimpl
 			OCIO_NAMESPACE::OptimizationFlags::OPTIMIZATION_NO_DYNAMIC_PROPERTIES
 		);
 	}
-#endif
+#endif // WITH_OCIO
 };
 
 struct FOpenColorIOGPUProcessorPimpl
@@ -53,7 +53,7 @@ struct FOpenColorIOGPUProcessorPimpl
 #if WITH_OCIO
 	OCIO_NAMESPACE::ConstGPUProcessorRcPtr Processor = nullptr;
 	OCIO_NAMESPACE::GpuShaderDescRcPtr ShaderDescription = nullptr;
-#endif
+#endif // WITH_OCIO
 };
 
 #if WITH_OCIO
@@ -99,7 +99,7 @@ namespace {
 			);
 	};
 }
-#endif
+#endif // WITH_OCIO
 
 TUniquePtr<FOpenColorIOConfigWrapper> FOpenColorIOConfigWrapper::CreateWorkingColorSpaceToInterchangeConfig()
 {
@@ -193,7 +193,7 @@ bool FOpenColorIOConfigWrapper::IsValid() const
 	return Pimpl->Config != nullptr;
 #else
 	return false;
-#endif
+#endif // WITH_OCIO
 }
 
 int32 FOpenColorIOConfigWrapper::GetNumColorSpaces() const
@@ -203,7 +203,7 @@ int32 FOpenColorIOConfigWrapper::GetNumColorSpaces() const
 	{
 		return Pimpl->Config->getNumColorSpaces(OCIO_NAMESPACE::SEARCH_REFERENCE_SPACE_ALL, OCIO_NAMESPACE::COLORSPACE_ACTIVE);
 	}
-#endif
+#endif  // WITH_OCIO
 
 	return 0;
 }
@@ -217,7 +217,7 @@ FString FOpenColorIOConfigWrapper::GetColorSpaceName(int32 Index) const
 
 		return StringCast<TCHAR>(ColorSpaceName).Get();
 	}
-#endif
+#endif // WITH_OCIO
 
 	return {};
 }
@@ -229,7 +229,7 @@ int32 FOpenColorIOConfigWrapper::GetColorSpaceIndex(const TCHAR* InColorSpaceNam
 	{
 		return Pimpl->Config->getIndexForColorSpace(StringCast<ANSICHAR>(InColorSpaceName).Get());
 	}
-#endif
+#endif // WITH_OCIO
 
 	return false;
 }
@@ -246,7 +246,7 @@ FString FOpenColorIOConfigWrapper::GetColorSpaceFamilyName(const TCHAR* InColorS
 			return StringCast<TCHAR>(ColorSpace->getFamily()).Get();
 		}
 	}
-#endif
+#endif // WITH_OCIO
 
 	return {};
 }
@@ -258,7 +258,7 @@ int32 FOpenColorIOConfigWrapper::GetNumDisplays() const
 	{
 		return Pimpl->Config->getNumDisplays();
 	}
-#endif
+#endif // WITH_OCIO
 
 	return 0;
 }
@@ -272,7 +272,7 @@ FString FOpenColorIOConfigWrapper::GetDisplayName(int32 Index) const
 
 		return StringCast<TCHAR>(DisplayName).Get();
 	}
-#endif
+#endif // WITH_OCIO
 
 	return {};
 }
@@ -284,7 +284,7 @@ int32 FOpenColorIOConfigWrapper::GetNumViews(const TCHAR* InDisplayName) const
 	{
 		return Pimpl->Config->getNumViews(StringCast<ANSICHAR>(InDisplayName).Get());
 	}
-#endif
+#endif // WITH_OCIO
 
 	return 0;
 }
@@ -298,7 +298,7 @@ FString FOpenColorIOConfigWrapper::GetViewName(const TCHAR* InDisplayName, int32
 
 		return StringCast<TCHAR>(ViewName).Get();
 	}
-#endif
+#endif // WITH_OCIO
 
 	return {};
 }
@@ -312,7 +312,7 @@ FString FOpenColorIOConfigWrapper::GetDisplayViewTransformName(const TCHAR* InDi
 
 		return StringCast<TCHAR>(TransformName).Get();
 	}
-#endif
+#endif // WITH_OCIO
 
 	return {};
 }
@@ -324,20 +324,13 @@ FString FOpenColorIOConfigWrapper::GetCacheID() const
 	{
 		return StringCast<TCHAR>(Pimpl->Config->getCacheID()).Get();
 	}
-#endif
+#endif // WITH_OCIO
 
 	return {};
 }
 
-
-FOpenColorIOProcessorWrapper::FOpenColorIOProcessorWrapper()
-	: Pimpl(MakePimpl<FOpenColorIOProcessorPimpl>())
-	, OwnerConfig(nullptr)
-	, WorkingColorSpaceTransformType(EOpenColorIOWorkingColorSpaceTransform::None)
-{ }
-
 FOpenColorIOProcessorWrapper::FOpenColorIOProcessorWrapper(
-	const FOpenColorIOConfigWrapper* InConfig,
+	const FOpenColorIOConfigWrapper& InConfig,
 	FStringView SourceColorSpace,
 	FStringView DestinationColorSpace,
 	EOpenColorIOWorkingColorSpaceTransform InWorkingColorSpaceTransformType,
@@ -346,12 +339,10 @@ FOpenColorIOProcessorWrapper::FOpenColorIOProcessorWrapper(
 	, OwnerConfig(InConfig)
 	, WorkingColorSpaceTransformType(InWorkingColorSpaceTransformType)
 {
-	ensure(OwnerConfig);
-
 #if WITH_OCIO
-	if (OwnerConfig && OwnerConfig->IsValid())
+	if (OwnerConfig.IsValid())
 	{
-		const OCIO_NAMESPACE::ConstConfigRcPtr& Config = OwnerConfig->Pimpl->Config;
+		const OCIO_NAMESPACE::ConstConfigRcPtr& Config = OwnerConfig.Pimpl->Config;
 		 OCIO_NAMESPACE::ContextRcPtr Context = Config->getCurrentContext()->createEditableCopy();
 
 		for (const TPair<FString, FString>& KeyValue : ContextKeyValues)
@@ -365,11 +356,11 @@ FOpenColorIOProcessorWrapper::FOpenColorIOProcessorWrapper(
 			StringCast<ANSICHAR>(DestinationColorSpace.GetData()).Get()
 		);
 	}
-#endif
+#endif // WITH_OCIO
 }
 
 FOpenColorIOProcessorWrapper::FOpenColorIOProcessorWrapper(
-	const FOpenColorIOConfigWrapper* InConfig,
+	const FOpenColorIOConfigWrapper& InConfig,
 	FStringView SourceColorSpace,
 	FStringView Display,
 	FStringView View,
@@ -380,27 +371,36 @@ FOpenColorIOProcessorWrapper::FOpenColorIOProcessorWrapper(
 	, OwnerConfig(InConfig)
 	, WorkingColorSpaceTransformType(InWorkingColorSpaceTransformType)
 {
-	ensure(OwnerConfig);
-
 #if WITH_OCIO
-	if (OwnerConfig && OwnerConfig->IsValid())
+#if !PLATFORM_EXCEPTIONS_DISABLED
+	try
+#endif
 	{
-		const OCIO_NAMESPACE::ConstConfigRcPtr& Config = OwnerConfig->Pimpl->Config;
-		OCIO_NAMESPACE::ContextRcPtr Context = Config->getCurrentContext()->createEditableCopy();
-
-		for (const TPair<FString, FString>& KeyValue : ContextKeyValues)
+		if (OwnerConfig.IsValid())
 		{
-			Context->setStringVar(TCHAR_TO_ANSI(*KeyValue.Key), TCHAR_TO_ANSI(*KeyValue.Value));
-		}
+			const OCIO_NAMESPACE::ConstConfigRcPtr& Config = OwnerConfig.Pimpl->Config;
+			OCIO_NAMESPACE::ContextRcPtr Context = Config->getCurrentContext()->createEditableCopy();
 
-		Pimpl->Processor = Config->getProcessor(
-			StringCast<ANSICHAR>(SourceColorSpace.GetData()).Get(),
-			StringCast<ANSICHAR>(Display.GetData()).Get(),
-			StringCast<ANSICHAR>(View.GetData()).Get(),
-			static_cast<OCIO_NAMESPACE::TransformDirection>(bInverseDirection)
-		);
+			for (const TPair<FString, FString>& KeyValue : ContextKeyValues)
+			{
+				Context->setStringVar(TCHAR_TO_ANSI(*KeyValue.Key), TCHAR_TO_ANSI(*KeyValue.Value));
+			}
+
+			Pimpl->Processor = Config->getProcessor(
+				StringCast<ANSICHAR>(SourceColorSpace.GetData()).Get(),
+				StringCast<ANSICHAR>(Display.GetData()).Get(),
+				StringCast<ANSICHAR>(View.GetData()).Get(),
+				static_cast<OCIO_NAMESPACE::TransformDirection>(bInverseDirection)
+			);
+		}
+	}
+#if !PLATFORM_EXCEPTIONS_DISABLED
+	catch (OCIO_NAMESPACE::Exception& Exc)
+	{
+		UE_LOG(LogOpenColorIOWrapper, Log, TEXT("Failed to create processor image. Error message: %s"), StringCast<TCHAR>(Exc.what()).Get());
 	}
 #endif
+#endif // WITH_OCIO
 }
 
 bool FOpenColorIOProcessorWrapper::IsValid() const
@@ -438,7 +438,7 @@ bool FOpenColorIOCPUProcessorWrapper::TransformImage(const FImageView& InOutImag
 			if (ImageDesc)
 			{
 				ConstConfigRcPtr	InterchangeConfig = IOpenColorIOWrapperModule::Get().GetWorkingColorSpaceToInterchangeConfig()->Pimpl->Config;
-				ConstConfigRcPtr Config = OwnerProcessor.OwnerConfig->Pimpl->Config;
+				ConstConfigRcPtr Config = OwnerProcessor.OwnerConfig.Pimpl->Config;
 
 				BitDepth BitDepth = ImageDesc->getBitDepth();
 
@@ -504,7 +504,7 @@ bool FOpenColorIOCPUProcessorWrapper::TransformImage(const FImageView& SrcImage,
 			if (SrcImageDesc && DestImageDesc)
 			{
 				ConstConfigRcPtr	InterchangeConfig = IOpenColorIOWrapperModule::Get().GetWorkingColorSpaceToInterchangeConfig()->Pimpl->Config;
-				ConstConfigRcPtr	Config = OwnerProcessor.OwnerConfig->Pimpl->Config;
+				ConstConfigRcPtr	Config = OwnerProcessor.OwnerConfig.Pimpl->Config;
 
 				BitDepth SrcBitDepth = SrcImageDesc->getBitDepth();
 				BitDepth DestBitDepth = DestImageDesc->getBitDepth();
@@ -594,11 +594,10 @@ FOpenColorIOGPUProcessorWrapper::FOpenColorIOGPUProcessorWrapper(const FOpenColo
 #if !PLATFORM_EXCEPTIONS_DISABLED
 	catch (OCIO_NAMESPACE::Exception& Exc)
 	{
-		//UE_LOG(LogOpenColorIOWrapper, Error, TEXT("Failed to fetch shader info for color transform %s. Error message: %s"), *GetTransformFriendlyName(), StringCast<TCHAR>(Exc.what()).Get());
 		UE_LOG(LogOpenColorIOWrapper, Error, TEXT("Failed to fetch shader info for color transform. Error message: %s"), StringCast<TCHAR>(Exc.what()).Get());
 	}
 #endif
-#endif //WITH_OCIO
+#endif // WITH_OCIO
 }
 
 bool FOpenColorIOGPUProcessorWrapper::IsValid() const
@@ -607,7 +606,7 @@ bool FOpenColorIOGPUProcessorWrapper::IsValid() const
 	return OwnerProcessor.IsValid() && GPUPimpl->Processor != nullptr && GPUPimpl->ShaderDescription != nullptr;
 #else
 	return false;
-#endif
+#endif // WITH_OCIO
 }
 
 bool FOpenColorIOGPUProcessorWrapper::GetShader(FString& OutShaderCode, FString& OutShaderCacheID) const
@@ -622,7 +621,7 @@ bool FOpenColorIOGPUProcessorWrapper::GetShader(FString& OutShaderCode, FString&
 
 		return true;
 	}
-#endif
+#endif // WITH_OCIO
 
 	return false;
 }
@@ -634,7 +633,7 @@ uint32 FOpenColorIOGPUProcessorWrapper::GetNum3DTextures() const
 	{
 		return GPUPimpl->ShaderDescription->getNum3DTextures();
 	}
-#endif
+#endif // WITH_OCIO
 
 	return 0;
 }
@@ -644,30 +643,41 @@ bool FOpenColorIOGPUProcessorWrapper::Get3DTexture(uint32 InIndex, FName& OutNam
 #if WITH_OCIO
 	ensure(InIndex < GetNum3DTextures());
 
-	if (IsValid())
+#if !PLATFORM_EXCEPTIONS_DISABLED
+	try
+#endif
 	{
-		const ANSICHAR* TextureName = nullptr;
-		const ANSICHAR* SamplerName = nullptr;
-		OCIO_NAMESPACE::Interpolation Interpolation = OCIO_NAMESPACE::INTERP_TETRAHEDRAL;
+		if (IsValid())
+		{
+			const ANSICHAR* TextureName = nullptr;
+			const ANSICHAR* SamplerName = nullptr;
+			OCIO_NAMESPACE::Interpolation Interpolation = OCIO_NAMESPACE::INTERP_TETRAHEDRAL;
 
-		// Read texture information
-		GPUPimpl->ShaderDescription->get3DTexture(InIndex, TextureName, SamplerName, OutEdgeLength, Interpolation);
+			// Read texture information
+			GPUPimpl->ShaderDescription->get3DTexture(InIndex, TextureName, SamplerName, OutEdgeLength, Interpolation);
 
-		// Read texture data
-		OutData = 0x0;
-		GPUPimpl->ShaderDescription->get3DTextureValues(InIndex, OutData);
+			// Read texture data
+			OutData = 0x0;
+			GPUPimpl->ShaderDescription->get3DTextureValues(InIndex, OutData);
 
-		OutName = FName(TextureName);
+			OutName = FName(TextureName);
 
-		OutTextureFilter = TF_Bilinear;
-		if (Interpolation == OCIO_NAMESPACE::Interpolation::INTERP_NEAREST || Interpolation == OCIO_NAMESPACE::Interpolation::INTERP_TETRAHEDRAL)
-		{
-			OutTextureFilter = TF_Nearest;
-		}
+			OutTextureFilter = TF_Bilinear;
+			if (Interpolation == OCIO_NAMESPACE::Interpolation::INTERP_NEAREST || Interpolation == OCIO_NAMESPACE::Interpolation::INTERP_TETRAHEDRAL)
+			{
+				OutTextureFilter = TF_Nearest;
+			}
 
-		return TextureName && OutEdgeLength > 0 && OutData;
+			return TextureName && OutEdgeLength > 0 && OutData;
+		}
+	}
+#if !PLATFORM_EXCEPTIONS_DISABLED
+	catch (OCIO_NAMESPACE::Exception& Exc)
+	{
+		UE_LOG(LogOpenColorIOWrapper, Error, TEXT("Failed to fetch 3d texture(s) info for color transform. Error message: %s"), StringCast<TCHAR>(Exc.what()).Get());
 	}
 #endif
+#endif //WITH_OCIO
 
 	return false;
 }
@@ -676,9 +686,9 @@ uint32 FOpenColorIOGPUProcessorWrapper::GetNumTextures() const
 #if WITH_OCIO
 	if (IsValid())
 	{
-		return GPUPimpl->ShaderDescription->getNumTextures();
+		return GPUPimpl->ShaderDescription->getNumTextures();  //noexcept
 	}
-#endif
+#endif // WITH_OCIO
 
 	return 0;
 }
@@ -688,27 +698,38 @@ bool FOpenColorIOGPUProcessorWrapper::GetTexture(uint32 InIndex, FName& OutName,
 	ensure(InIndex < GetNumTextures());
 
 #if WITH_OCIO
-	if (IsValid())
+#if !PLATFORM_EXCEPTIONS_DISABLED
+	try
+#endif
 	{
-		const ANSICHAR* TextureName = nullptr;
-		const ANSICHAR* SamplerName = nullptr;
-		OCIO_NAMESPACE::GpuShaderDesc::TextureType Channel = OCIO_NAMESPACE::GpuShaderDesc::TEXTURE_RGB_CHANNEL;
-		OCIO_NAMESPACE::Interpolation Interpolation = OCIO_NAMESPACE::Interpolation::INTERP_LINEAR;
+		if (IsValid())
+		{
+			const ANSICHAR* TextureName = nullptr;
+			const ANSICHAR* SamplerName = nullptr;
+			OCIO_NAMESPACE::GpuShaderDesc::TextureType Channel = OCIO_NAMESPACE::GpuShaderDesc::TEXTURE_RGB_CHANNEL;
+			OCIO_NAMESPACE::Interpolation Interpolation = OCIO_NAMESPACE::Interpolation::INTERP_LINEAR;
 
-		// Read texture information
-		GPUPimpl->ShaderDescription->getTexture(InIndex, TextureName, SamplerName, OutWidth, OutHeight, Channel, Interpolation);
+			// Read texture information
+			GPUPimpl->ShaderDescription->getTexture(InIndex, TextureName, SamplerName, OutWidth, OutHeight, Channel, Interpolation);
 
-		// Read texture data
-		OutData = 0x0;
-		GPUPimpl->ShaderDescription->getTextureValues(InIndex, OutData);
+			// Read texture data
+			OutData = 0x0;
+			GPUPimpl->ShaderDescription->getTextureValues(InIndex, OutData);
 
-		OutName = FName(TextureName);
-		OutTextureFilter = Interpolation == OCIO_NAMESPACE::Interpolation::INTERP_NEAREST ? TF_Nearest : TF_Bilinear;
-		bOutRedChannelOnly = Channel == OCIO_NAMESPACE::GpuShaderCreator::TEXTURE_RED_CHANNEL;
+			OutName = FName(TextureName);
+			OutTextureFilter = Interpolation == OCIO_NAMESPACE::Interpolation::INTERP_NEAREST ? TF_Nearest : TF_Bilinear;
+			bOutRedChannelOnly = Channel == OCIO_NAMESPACE::GpuShaderCreator::TEXTURE_RED_CHANNEL;
 
-		return TextureName && OutWidth > 0 && OutHeight > 0 && OutData;
+			return TextureName && OutWidth > 0 && OutHeight > 0 && OutData;
+		}
+	}
+#if !PLATFORM_EXCEPTIONS_DISABLED
+	catch (OCIO_NAMESPACE::Exception& Exc)
+	{
+		UE_LOG(LogOpenColorIOWrapper, Error, TEXT("Failed to fetch texture(s) info for color transform. Error message: %s"), StringCast<TCHAR>(Exc.what()).Get());
 	}
 #endif
+#endif // WITH_OCIO
 
 	return false;
 }
@@ -720,7 +741,7 @@ FString FOpenColorIOGPUProcessorWrapper::GetCacheID() const
 	{
 		return StringCast<TCHAR>(GPUPimpl->Processor->getCacheID()).Get();
 	}
-#endif
+#endif //WITH_OCIO
 
 	return {};
 }
diff --git a/Engine/Plugins/Compositing/OpenColorIO/Source/OpenColorIOWrapper/Public/OpenColorIOWrapper.h b/Engine/Plugins/Compositing/OpenColorIO/Source/OpenColorIOWrapper/Public/OpenColorIOWrapper.h
--- a/Engine/Plugins/Compositing/OpenColorIO/Source/OpenColorIOWrapper/Public/OpenColorIOWrapper.h
+++ b/Engine/Plugins/Compositing/OpenColorIO/Source/OpenColorIOWrapper/Public/OpenColorIOWrapper.h
@@ -105,9 +105,6 @@ private:
 class OPENCOLORIOWRAPPER_API FOpenColorIOProcessorWrapper final
 {
 public:
-	/** Constructor. */
-	FOpenColorIOProcessorWrapper();
-	
 	/**
 	* Constructor.
 	*
@@ -118,7 +115,7 @@ public:
 	* @param ContextKeyValues (Optional) ContextKeyValues Additional context modifiers.
 	*/
 	FOpenColorIOProcessorWrapper(
-		const FOpenColorIOConfigWrapper* InConfig,
+		const FOpenColorIOConfigWrapper& InConfig,
 		FStringView SourceColorSpace,
 		FStringView DestinationColorSpace,
 		EOpenColorIOWorkingColorSpaceTransform WorkingColorSpaceTransformType = EOpenColorIOWorkingColorSpaceTransform::None,
@@ -136,7 +133,7 @@ public:
 	* @param ContextKeyValues (Optional) ContextKeyValues Additional context modifiers.
 	*/
 	FOpenColorIOProcessorWrapper(
-		const FOpenColorIOConfigWrapper* InConfig,
+		const FOpenColorIOConfigWrapper& InConfig,
 		FStringView SourceColorSpace,
 		FStringView Display,
 		FStringView View,
@@ -151,7 +148,7 @@ private:
 
 	TPimplPtr<struct FOpenColorIOProcessorPimpl> Pimpl;
 
-	const FOpenColorIOConfigWrapper* OwnerConfig;
+	const FOpenColorIOConfigWrapper& OwnerConfig;
 
 	EOpenColorIOWorkingColorSpaceTransform WorkingColorSpaceTransformType;
 
