OCIO: Fix static analysis warning, and improve processor memory safety by simply copying the parent processor on construction (using pimpl deep copy mode, which will simply increase the reference count of the native OpenColorIO object).

#rnx

diff --git a/Engine/Plugins/Compositing/OpenColorIO/Source/OpenColorIO/Private/OpenColorIOColorTransform.cpp b/Engine/Plugins/Compositing/OpenColorIO/Source/OpenColorIO/Private/OpenColorIOColorTransform.cpp
--- a/Engine/Plugins/Compositing/OpenColorIO/Source/OpenColorIO/Private/OpenColorIOColorTransform.cpp
+++ b/Engine/Plugins/Compositing/OpenColorIO/Source/OpenColorIO/Private/OpenColorIOColorTransform.cpp
@@ -32,13 +32,12 @@ namespace
 		check(InTransform->ConfigurationOwner);
 
 		const FOpenColorIOConfigWrapper* ConfigWrapper = InTransform->ConfigurationOwner->GetConfigWrapper();
-		ensure(ConfigWrapper);
 
 		EOpenColorIOViewTransformDirection DisplayViewDirection;
 		if (InTransform->GetDisplayViewDirection(DisplayViewDirection))
 		{
 			return FOpenColorIOProcessorWrapper(
-				*ConfigWrapper,
+				ConfigWrapper,
 				InTransform->SourceColorSpace,
 				InTransform->Display,
 				InTransform->View,
@@ -49,7 +48,7 @@ namespace
 		else
 		{
 			return FOpenColorIOProcessorWrapper(
-				*ConfigWrapper,
+				ConfigWrapper,
 				InTransform->SourceColorSpace,
 				InTransform->DestinationColorSpace,
 				InWCSTrasnformType,
diff --git a/Engine/Plugins/Compositing/OpenColorIO/Source/OpenColorIOWrapper/Private/OpenColorIOWrapper.cpp b/Engine/Plugins/Compositing/OpenColorIO/Source/OpenColorIOWrapper/Private/OpenColorIOWrapper.cpp
--- a/Engine/Plugins/Compositing/OpenColorIO/Source/OpenColorIOWrapper/Private/OpenColorIOWrapper.cpp
+++ b/Engine/Plugins/Compositing/OpenColorIO/Source/OpenColorIOWrapper/Private/OpenColorIOWrapper.cpp
@@ -139,7 +139,7 @@ TUniquePtr<FOpenColorIOConfigWrapper> FOpenColorIOConfigWrapper::CreateWorkingCo
 }
 
 FOpenColorIOConfigWrapper::FOpenColorIOConfigWrapper()
-	: Pimpl(MakePimpl<FOpenColorIOConfigPimpl>())
+	: Pimpl(MakePimpl<FOpenColorIOConfigPimpl, EPimplPtrMode::DeepCopy>())
 { }
 
 FOpenColorIOConfigWrapper::FOpenColorIOConfigWrapper(FStringView FilePath, FInitializationOptions InOptions)
@@ -330,19 +330,19 @@ FString FOpenColorIOConfigWrapper::GetCacheID() const
 }
 
 FOpenColorIOProcessorWrapper::FOpenColorIOProcessorWrapper(
-	const FOpenColorIOConfigWrapper& InConfig,
+	const FOpenColorIOConfigWrapper* InConfig,
 	FStringView SourceColorSpace,
 	FStringView DestinationColorSpace,
 	EOpenColorIOWorkingColorSpaceTransform InWorkingColorSpaceTransformType,
 	const TMap<FString, FString>& ContextKeyValues)
-	: Pimpl(MakePimpl<FOpenColorIOProcessorPimpl>())
+	: Pimpl(MakePimpl<FOpenColorIOProcessorPimpl, EPimplPtrMode::DeepCopy>())
 	, OwnerConfig(InConfig)
 	, WorkingColorSpaceTransformType(InWorkingColorSpaceTransformType)
 {
 #if WITH_OCIO
-	if (OwnerConfig.IsValid())
+	if (OwnerConfig != nullptr && OwnerConfig->IsValid())
 	{
-		const OCIO_NAMESPACE::ConstConfigRcPtr& Config = OwnerConfig.Pimpl->Config;
+		const OCIO_NAMESPACE::ConstConfigRcPtr& Config = OwnerConfig->Pimpl->Config;
 		 OCIO_NAMESPACE::ContextRcPtr Context = Config->getCurrentContext()->createEditableCopy();
 
 		for (const TPair<FString, FString>& KeyValue : ContextKeyValues)
@@ -360,14 +360,14 @@ FOpenColorIOProcessorWrapper::FOpenColorIOProcessorWrapper(
 }
 
 FOpenColorIOProcessorWrapper::FOpenColorIOProcessorWrapper(
-	const FOpenColorIOConfigWrapper& InConfig,
+	const FOpenColorIOConfigWrapper* InConfig,
 	FStringView SourceColorSpace,
 	FStringView Display,
 	FStringView View,
 	bool bInverseDirection,
 	EOpenColorIOWorkingColorSpaceTransform InWorkingColorSpaceTransformType,
 	const TMap<FString, FString>& ContextKeyValues)
-	: Pimpl(MakePimpl<FOpenColorIOProcessorPimpl>())
+	: Pimpl(MakePimpl<FOpenColorIOProcessorPimpl, EPimplPtrMode::DeepCopy>())
 	, OwnerConfig(InConfig)
 	, WorkingColorSpaceTransformType(InWorkingColorSpaceTransformType)
 {
@@ -376,9 +376,9 @@ FOpenColorIOProcessorWrapper::FOpenColorIOProcessorWrapper(
 	try
 #endif
 	{
-		if (OwnerConfig.IsValid())
+		if (OwnerConfig != nullptr && OwnerConfig->IsValid())
 		{
-			const OCIO_NAMESPACE::ConstConfigRcPtr& Config = OwnerConfig.Pimpl->Config;
+			const OCIO_NAMESPACE::ConstConfigRcPtr& Config = OwnerConfig->Pimpl->Config;
 			OCIO_NAMESPACE::ContextRcPtr Context = Config->getCurrentContext()->createEditableCopy();
 
 			for (const TPair<FString, FString>& KeyValue : ContextKeyValues)
@@ -406,20 +406,20 @@ FOpenColorIOProcessorWrapper::FOpenColorIOProcessorWrapper(
 bool FOpenColorIOProcessorWrapper::IsValid() const
 {
 #if WITH_OCIO
-	return Pimpl->Processor != nullptr;
+	return OwnerConfig != nullptr && Pimpl->Processor != nullptr;
 #else
 	return false;
 #endif
 }
 
-FOpenColorIOCPUProcessorWrapper::FOpenColorIOCPUProcessorWrapper(const FOpenColorIOProcessorWrapper& InOwner)
-	: OwnerProcessor(InOwner)
+FOpenColorIOCPUProcessorWrapper::FOpenColorIOCPUProcessorWrapper(FOpenColorIOProcessorWrapper InProcessor)
+	: ParentProcessor(MoveTemp(InProcessor))
 {
 }
 
 bool FOpenColorIOCPUProcessorWrapper::IsValid() const
 {
-	return OwnerProcessor.IsValid();
+	return ParentProcessor.IsValid();
 }
 
 bool FOpenColorIOCPUProcessorWrapper::TransformImage(const FImageView& InOutImage) const
@@ -432,18 +432,18 @@ bool FOpenColorIOCPUProcessorWrapper::TransformImage(const FImageView& InOutImag
 	{
 		using namespace OCIO_NAMESPACE;
 
-		if (OwnerProcessor.IsValid())
+		if (ParentProcessor.IsValid())
 		{
 			TUniquePtr<PackedImageDesc> ImageDesc = GetImageDesc(InOutImage);
 			if (ImageDesc)
 			{
 				ConstConfigRcPtr	InterchangeConfig = IOpenColorIOWrapperModule::Get().GetWorkingColorSpaceToInterchangeConfig()->Pimpl->Config;
-				ConstConfigRcPtr Config = OwnerProcessor.OwnerConfig.Pimpl->Config;
+				ConstConfigRcPtr Config = ParentProcessor.OwnerConfig->Pimpl->Config;
 
 				BitDepth BitDepth = ImageDesc->getBitDepth();
 
 				// Conditionally apply a conversion from the working color space to interchange space
-				if (OwnerProcessor.WorkingColorSpaceTransformType == EOpenColorIOWorkingColorSpaceTransform::Source)
+				if (ParentProcessor.WorkingColorSpaceTransformType == EOpenColorIOWorkingColorSpaceTransform::Source)
 				{
 					ConstProcessorRcPtr	InterchangeProcessor = InterchangeConfig->GetProcessorFromConfigs(
 						InterchangeConfig,
@@ -456,11 +456,11 @@ bool FOpenColorIOCPUProcessorWrapper::TransformImage(const FImageView& InOutImag
 				}
 
 				// Apply the main color transformation
-				ConstCPUProcessorRcPtr CPUProcessor = OwnerProcessor.Pimpl->Processor->getOptimizedCPUProcessor(BitDepth, BitDepth, OPTIMIZATION_DEFAULT);
+				ConstCPUProcessorRcPtr CPUProcessor = ParentProcessor.Pimpl->Processor->getOptimizedCPUProcessor(BitDepth, BitDepth, OPTIMIZATION_DEFAULT);
 				CPUProcessor->apply(*ImageDesc);
 
 				// Conditionally apply a conversion from the interchange space to the working color space
-				if (OwnerProcessor.WorkingColorSpaceTransformType == EOpenColorIOWorkingColorSpaceTransform::Destination)
+				if (ParentProcessor.WorkingColorSpaceTransformType == EOpenColorIOWorkingColorSpaceTransform::Destination)
 				{
 					ConstProcessorRcPtr	InterchangeProcessor = InterchangeConfig->GetProcessorFromConfigs(
 						Config,
@@ -497,20 +497,20 @@ bool FOpenColorIOCPUProcessorWrapper::TransformImage(const FImageView& SrcImage,
 	{
 		using namespace OCIO_NAMESPACE;
 
-		if (OwnerProcessor.IsValid())
+		if (ParentProcessor.IsValid())
 		{
 			TUniquePtr<PackedImageDesc> SrcImageDesc = GetImageDesc(SrcImage);
 			TUniquePtr<PackedImageDesc> DestImageDesc = GetImageDesc(DestImage);
 			if (SrcImageDesc && DestImageDesc)
 			{
 				ConstConfigRcPtr	InterchangeConfig = IOpenColorIOWrapperModule::Get().GetWorkingColorSpaceToInterchangeConfig()->Pimpl->Config;
-				ConstConfigRcPtr	Config = OwnerProcessor.OwnerConfig.Pimpl->Config;
+				ConstConfigRcPtr	Config = ParentProcessor.OwnerConfig->Pimpl->Config;
 
 				BitDepth SrcBitDepth = SrcImageDesc->getBitDepth();
 				BitDepth DestBitDepth = DestImageDesc->getBitDepth();
 
 				// Conditionally apply a conversion from the working color space to interchange space
-				if (OwnerProcessor.WorkingColorSpaceTransformType == EOpenColorIOWorkingColorSpaceTransform::Source)
+				if (ParentProcessor.WorkingColorSpaceTransformType == EOpenColorIOWorkingColorSpaceTransform::Source)
 				{
 					ConstProcessorRcPtr	InterchangeProcessor = InterchangeConfig->GetProcessorFromConfigs(
 						InterchangeConfig,
@@ -524,11 +524,11 @@ bool FOpenColorIOCPUProcessorWrapper::TransformImage(const FImageView& SrcImage,
 
 
 				// Apply the main color transformation
-				ConstCPUProcessorRcPtr CPUProcessor = OwnerProcessor.Pimpl->Processor->getOptimizedCPUProcessor(SrcBitDepth, DestBitDepth, OPTIMIZATION_DEFAULT);
+				ConstCPUProcessorRcPtr CPUProcessor = ParentProcessor.Pimpl->Processor->getOptimizedCPUProcessor(SrcBitDepth, DestBitDepth, OPTIMIZATION_DEFAULT);
 				CPUProcessor->apply(*SrcImageDesc, *DestImageDesc);
 
 				// Conditionally apply a conversion from the interchange space to the working color space
-				if (OwnerProcessor.WorkingColorSpaceTransformType == EOpenColorIOWorkingColorSpaceTransform::Destination)
+				if (ParentProcessor.WorkingColorSpaceTransformType == EOpenColorIOWorkingColorSpaceTransform::Destination)
 				{
 					ConstProcessorRcPtr	InterchangeProcessor = InterchangeConfig->GetProcessorFromConfigs(
 						Config,
@@ -557,9 +557,9 @@ bool FOpenColorIOCPUProcessorWrapper::TransformImage(const FImageView& SrcImage,
 
 
 
-FOpenColorIOGPUProcessorWrapper::FOpenColorIOGPUProcessorWrapper(const FOpenColorIOProcessorWrapper& InOwner, FInitializationOptions InShaderParams)
-	: OwnerProcessor(InOwner)
-	, GPUPimpl(MakePimpl<FOpenColorIOGPUProcessorPimpl>())
+FOpenColorIOGPUProcessorWrapper::FOpenColorIOGPUProcessorWrapper(FOpenColorIOProcessorWrapper InProcessor, FInitializationOptions InShaderParams)
+	: ParentProcessor(MoveTemp(InProcessor))
+	, GPUPimpl(MakePimpl<FOpenColorIOGPUProcessorPimpl, EPimplPtrMode::DeepCopy>())
 {
 #if WITH_OCIO
 	TRACE_CPUPROFILER_EVENT_SCOPE(FOpenColorIOGPUProcessorWrapper::FOpenColorIOGPUProcessorWrapper)
@@ -580,11 +580,11 @@ FOpenColorIOGPUProcessorWrapper::FOpenColorIOGPUProcessorWrapper(const FOpenColo
 		if (InShaderParams.bIsLegacy)
 		{
 			unsigned int EdgeLength = static_cast<unsigned int>(OpenColorIOWrapper::Legacy3dEdgeLength);
-			GPUProcessor = OwnerProcessor.Pimpl->Processor->getOptimizedLegacyGPUProcessor(OptFlags, EdgeLength);
+			GPUProcessor = ParentProcessor.Pimpl->Processor->getOptimizedLegacyGPUProcessor(OptFlags, EdgeLength);
 		}
 		else
 		{
-			GPUProcessor = OwnerProcessor.Pimpl->Processor->getOptimizedGPUProcessor(FOpenColorIOProcessorPimpl::GetOptimizationFlags());
+			GPUProcessor = ParentProcessor.Pimpl->Processor->getOptimizedGPUProcessor(FOpenColorIOProcessorPimpl::GetOptimizationFlags());
 		}
 		GPUProcessor->extractGpuShaderInfo(ShaderDescription);
 
@@ -603,7 +603,7 @@ FOpenColorIOGPUProcessorWrapper::FOpenColorIOGPUProcessorWrapper(const FOpenColo
 bool FOpenColorIOGPUProcessorWrapper::IsValid() const
 {
 #if WITH_OCIO
-	return OwnerProcessor.IsValid() && GPUPimpl->Processor != nullptr && GPUPimpl->ShaderDescription != nullptr;
+	return ParentProcessor.IsValid() && GPUPimpl->Processor != nullptr && GPUPimpl->ShaderDescription != nullptr;
 #else
 	return false;
 #endif // WITH_OCIO
diff --git a/Engine/Plugins/Compositing/OpenColorIO/Source/OpenColorIOWrapper/Public/OpenColorIOWrapper.h b/Engine/Plugins/Compositing/OpenColorIO/Source/OpenColorIOWrapper/Public/OpenColorIOWrapper.h
--- a/Engine/Plugins/Compositing/OpenColorIO/Source/OpenColorIOWrapper/Public/OpenColorIOWrapper.h
+++ b/Engine/Plugins/Compositing/OpenColorIO/Source/OpenColorIOWrapper/Public/OpenColorIOWrapper.h
@@ -94,7 +94,7 @@ private:
 	/** Convenience to create a config between the working color space and the default interchange one. */
 	static TUniquePtr<FOpenColorIOConfigWrapper> CreateWorkingColorSpaceToInterchangeConfig();
 
-	TPimplPtr<struct FOpenColorIOConfigPimpl> Pimpl;
+	TPimplPtr<struct FOpenColorIOConfigPimpl, EPimplPtrMode::DeepCopy> Pimpl;
 
 	friend class FOpenColorIOWrapperModule;
 	friend class FOpenColorIOProcessorWrapper;
@@ -115,7 +115,7 @@ public:
 	* @param ContextKeyValues (Optional) ContextKeyValues Additional context modifiers.
 	*/
 	FOpenColorIOProcessorWrapper(
-		const FOpenColorIOConfigWrapper& InConfig,
+		const FOpenColorIOConfigWrapper* InConfig,
 		FStringView SourceColorSpace,
 		FStringView DestinationColorSpace,
 		EOpenColorIOWorkingColorSpaceTransform WorkingColorSpaceTransformType = EOpenColorIOWorkingColorSpaceTransform::None,
@@ -133,7 +133,7 @@ public:
 	* @param ContextKeyValues (Optional) ContextKeyValues Additional context modifiers.
 	*/
 	FOpenColorIOProcessorWrapper(
-		const FOpenColorIOConfigWrapper& InConfig,
+		const FOpenColorIOConfigWrapper* InConfig,
 		FStringView SourceColorSpace,
 		FStringView Display,
 		FStringView View,
@@ -146,9 +146,9 @@ public:
 
 private:
 
-	TPimplPtr<struct FOpenColorIOProcessorPimpl> Pimpl;
+	TPimplPtr<struct FOpenColorIOProcessorPimpl, EPimplPtrMode::DeepCopy> Pimpl;
 
-	const FOpenColorIOConfigWrapper& OwnerConfig;
+	const FOpenColorIOConfigWrapper* OwnerConfig;
 
 	EOpenColorIOWorkingColorSpaceTransform WorkingColorSpaceTransformType;
 
@@ -163,9 +163,9 @@ public:
 	/**
 	* Constructor.
 	*
-	* @param InOwner Owner processor.
+	* @param InProcessor Parent processor.
 	*/
-	FOpenColorIOCPUProcessorWrapper(const FOpenColorIOProcessorWrapper& InOwner);
+	FOpenColorIOCPUProcessorWrapper(FOpenColorIOProcessorWrapper InProcessor);
 
 	/** Valid when the processor has been successfully created and isn't null. */
 	bool IsValid() const;
@@ -178,7 +178,7 @@ public:
 
 private:
 
-	const FOpenColorIOProcessorWrapper& OwnerProcessor;
+	const FOpenColorIOProcessorWrapper ParentProcessor;
 };
 
 class OPENCOLORIOWRAPPER_API FOpenColorIOGPUProcessorWrapper final
@@ -195,11 +195,11 @@ public:
 	/**
 	* Constructor.
 	*
-	* @param InOwner Owner processor.
+	* @param InProcessor Parent processor.
 	* @param InOptions Initialization options.
 	*/
 	FOpenColorIOGPUProcessorWrapper(
-		const FOpenColorIOProcessorWrapper& InOwner,
+		FOpenColorIOProcessorWrapper InProcessor,
 		FInitializationOptions InOptions = FInitializationOptions()
 	);
 
@@ -252,7 +252,7 @@ public:
 
 private:
 
-	const FOpenColorIOProcessorWrapper& OwnerProcessor;
+	const FOpenColorIOProcessorWrapper ParentProcessor;
 
-	TPimplPtr<struct FOpenColorIOGPUProcessorPimpl> GPUPimpl;
+	TPimplPtr<struct FOpenColorIOGPUProcessorPimpl, EPimplPtrMode::DeepCopy> GPUPimpl;
 };
