OCIO: Preparation for defining an engine config.

Minor cleanup of which default strings are exposed, and processor constructors.

diff --git a/Engine/Plugins/Compositing/OpenColorIO/Source/OpenColorIO/Private/OpenColorIOColorTransform.cpp b/Engine/Plugins/Compositing/OpenColorIO/Source/OpenColorIO/Private/OpenColorIOColorTransform.cpp
--- a/Engine/Plugins/Compositing/OpenColorIO/Source/OpenColorIO/Private/OpenColorIOColorTransform.cpp
+++ b/Engine/Plugins/Compositing/OpenColorIO/Source/OpenColorIO/Private/OpenColorIOColorTransform.cpp
@@ -26,7 +26,7 @@
 namespace
 {
 	// Note: Since OpenColorIO caches processors automatically, we can recreate them without significant additional costs.
-	FOpenColorIOProcessorWrapper GetTransformProcessor(const UOpenColorIOColorTransform* InTransform, EOpenColorIOWorkingColorSpaceTransform InWCSTrasnformType)
+	FOpenColorIOProcessorWrapper GetTransformProcessor(const UOpenColorIOColorTransform* InTransform)
 	{
 		check(InTransform);
 		check(InTransform->ConfigurationOwner);
@@ -42,7 +42,6 @@ namespace
 				InTransform->Display,
 				InTransform->View,
 				DisplayViewDirection == EOpenColorIOViewTransformDirection::Inverse,
-				InWCSTrasnformType,
 				InTransform->GetContextKeyValues());
 		}
 		else
@@ -51,7 +50,6 @@ namespace
 				ConfigWrapper,
 				InTransform->SourceColorSpace,
 				InTransform->DestinationColorSpace,
-				InWCSTrasnformType,
 				InTransform->GetContextKeyValues());
 		}
 	}
@@ -160,11 +158,11 @@ bool UOpenColorIOColorTransform::Initialize(UOpenColorIOConfiguration* InOwner,
 	ContextKeyValues = InContextKeyValues;
 	
 #if WITH_EDITOR
-	if (InSourceColorSpace == OpenColorIOWrapper::WorkingColorSpaceName)
+	if (InSourceColorSpace == OpenColorIOWrapper::GetWorkingColorSpaceName())
 	{
 		WorkingColorSpaceTransformType = EOpenColorIOWorkingColorSpaceTransform::Source;
 	}
-	else if(InDestinationColorSpace == OpenColorIOWrapper::WorkingColorSpaceName)
+	else if(InDestinationColorSpace == OpenColorIOWrapper::GetWorkingColorSpaceName())
 	{
 		WorkingColorSpaceTransformType = EOpenColorIOWorkingColorSpaceTransform::Destination;
 	}
@@ -180,7 +178,7 @@ bool UOpenColorIOColorTransform::Initialize(UOpenColorIOConfiguration* InOwner,
 	ContextKeyValues = InContextKeyValues;
 
 #if WITH_EDITOR
-	if (InSourceColorSpace == OpenColorIOWrapper::WorkingColorSpaceName)
+	if (InSourceColorSpace == OpenColorIOWrapper::GetWorkingColorSpaceName())
 	{
 		if (InDirection == EOpenColorIOViewTransformDirection::Forward)
 		{
@@ -265,7 +263,7 @@ void UOpenColorIOColorTransform::CacheResourceTextures()
 #if WITH_EDITOR
 	if (Textures.IsEmpty())
 	{
-		const FOpenColorIOProcessorWrapper TransformProcessor = GetTransformProcessor(this, WorkingColorSpaceTransformType);
+		const FOpenColorIOProcessorWrapper TransformProcessor = GetTransformProcessor(this);
 		if (TransformProcessor.IsValid())
 		{
 			const FOpenColorIOGPUProcessorWrapper GPUProcessor = FOpenColorIOGPUProcessorWrapper(TransformProcessor);
@@ -476,7 +474,7 @@ bool UOpenColorIOColorTransform::IsTransform(const FString& InSourceColorSpace,
 #if WITH_EDITOR
 bool UOpenColorIOColorTransform::EditorTransformImage(const FImageView& InOutImage) const
 {
-	const FOpenColorIOProcessorWrapper TransformProcessor = GetTransformProcessor(this, WorkingColorSpaceTransformType);
+	const FOpenColorIOProcessorWrapper TransformProcessor = GetTransformProcessor(this);
 	if (TransformProcessor.IsValid())
 	{
 		const FOpenColorIOCPUProcessorWrapper CPUProcessor = FOpenColorIOCPUProcessorWrapper(TransformProcessor);
@@ -489,7 +487,7 @@ bool UOpenColorIOColorTransform::EditorTransformImage(const FImageView& InOutIma
 
 bool UOpenColorIOColorTransform::EditorTransformImage(const FImageView& SrcImage, const FImageView& DestImage) const
 {
-	const FOpenColorIOProcessorWrapper TransformProcessor = GetTransformProcessor(this, WorkingColorSpaceTransformType);
+	const FOpenColorIOProcessorWrapper TransformProcessor = GetTransformProcessor(this);
 	if (TransformProcessor.IsValid())
 	{
 		const FOpenColorIOCPUProcessorWrapper CPUProcessor = FOpenColorIOCPUProcessorWrapper(TransformProcessor);
@@ -588,7 +586,7 @@ FString UOpenColorIOColorTransform::GetTransformFriendlyName()
 bool UOpenColorIOColorTransform::UpdateShaderInfo(FString& OutShaderCodeHash, FString& OutShaderCode, FString& OutRawConfigHash)
 {
 #if WITH_EDITOR
-	const FOpenColorIOProcessorWrapper Processor = GetTransformProcessor(this, WorkingColorSpaceTransformType);
+	const FOpenColorIOProcessorWrapper Processor = GetTransformProcessor(this);
 	const FOpenColorIOGPUProcessorWrapper GPUProcessor = FOpenColorIOGPUProcessorWrapper(Processor);
 
 	if (GPUProcessor.IsValid())
diff --git a/Engine/Plugins/Compositing/OpenColorIO/Source/OpenColorIOWrapper/Private/OpenColorIOWrapper.cpp b/Engine/Plugins/Compositing/OpenColorIO/Source/OpenColorIOWrapper/Private/OpenColorIOWrapper.cpp
--- a/Engine/Plugins/Compositing/OpenColorIO/Source/OpenColorIOWrapper/Private/OpenColorIOWrapper.cpp
+++ b/Engine/Plugins/Compositing/OpenColorIO/Source/OpenColorIOWrapper/Private/OpenColorIOWrapper.cpp
@@ -16,6 +16,24 @@ THIRD_PARTY_INCLUDES_END
 
 #define LOCTEXT_NAMESPACE "OpenColorIOWrapper"
 
+namespace OpenColorIOWrapper
+{
+	/** OCIO_NAMESPACE::ROLE_INTERCHANGE_SCENE equivalent, since we currently cannot delay-load this definition. */
+	constexpr const ANSICHAR* GetInterchangeName()
+	{
+		return "aces_interchange";
+	}
+
+	// Build routine since there is no FAnsiString
+	TUniquePtr<ANSICHAR[]> MakeAnsiString(const TCHAR* Str)
+	{
+		int32 Num = FCString::Strlen(Str);
+		TUniquePtr<ANSICHAR[]> Ret =  MakeUnique<ANSICHAR[]>(Num);
+		FMemory::Memcpy(Ret.Get(), StringCast<ANSICHAR>(Str).Get(), Num);
+		return Ret;
+	}
+}
+
 const TCHAR* OpenColorIOWrapper::GetVersion()
 {
 #if WITH_OCIO
@@ -116,7 +134,7 @@ TUniquePtr<FOpenColorIOConfigWrapper> FOpenColorIOConfigWrapper::CreateWorkingCo
 	AP0->setEncoding("scene-linear");
 
 	ColorSpaceRcPtr WCS = ColorSpace::Create();
-	WCS->setName(OpenColorIOWrapper::WorkingColorSpaceName);
+	WCS->setName(StringCast<ANSICHAR>(OpenColorIOWrapper::GetWorkingColorSpaceName()).Get());
 	WCS->setBitDepth(BIT_DEPTH_F32);
 	WCS->setEncoding("scene-linear");
 
@@ -142,7 +160,7 @@ FOpenColorIOConfigWrapper::FOpenColorIOConfigWrapper()
 	: Pimpl(MakePimpl<FOpenColorIOConfigPimpl, EPimplPtrMode::DeepCopy>())
 { }
 
-FOpenColorIOConfigWrapper::FOpenColorIOConfigWrapper(FStringView FilePath, FInitializationOptions InOptions)
+FOpenColorIOConfigWrapper::FOpenColorIOConfigWrapper(FStringView InFilePath, FInitializationOptions InOptions)
 	: FOpenColorIOConfigWrapper()
 {
 #if WITH_OCIO
@@ -153,17 +171,18 @@ FOpenColorIOConfigWrapper::FOpenColorIOConfigWrapper(FStringView FilePath, FInit
 	{
 		using namespace OCIO_NAMESPACE;
 
-		ConstConfigRcPtr NewConfig = Config::CreateFromFile(StringCast<ANSICHAR>(FilePath.GetData()).Get());
+		ConstConfigRcPtr NewConfig = Config::CreateFromFile(StringCast<ANSICHAR>(InFilePath.GetData()).Get());
 
 		if (NewConfig && InOptions.bAddWorkingColorSpace)
 		{
-			ConstColorSpaceRcPtr InterchangeCS = NewConfig->getColorSpace(NewConfig->getCanonicalName(OpenColorIOWrapper::InterchangeName));
+			const TUniquePtr<ANSICHAR[]> AnsiWorkingColorSpaceName = OpenColorIOWrapper::MakeAnsiString(OpenColorIOWrapper::GetWorkingColorSpaceName());
+			ConstColorSpaceRcPtr InterchangeCS = NewConfig->getColorSpace(NewConfig->getCanonicalName(OpenColorIOWrapper::GetInterchangeName()));
 
 			// When the aces interchange color space is present, we add the working color space as an additional option.
-			if (InterchangeCS != nullptr && NewConfig->getColorSpace(OpenColorIOWrapper::WorkingColorSpaceName) == nullptr)
+			if (InterchangeCS != nullptr && NewConfig->getColorSpace(AnsiWorkingColorSpaceName.Get()) == nullptr)
 			{
 				ColorSpaceRcPtr WorkingCS = InterchangeCS->createEditableCopy();
-				WorkingCS->setName(OpenColorIOWrapper::WorkingColorSpaceName);
+				WorkingCS->setName(AnsiWorkingColorSpaceName.Get());
 				WorkingCS->setFamily("UE");
 				WorkingCS->clearAliases();
 
@@ -179,7 +198,7 @@ FOpenColorIOConfigWrapper::FOpenColorIOConfigWrapper(FStringView FilePath, FInit
 #if !PLATFORM_EXCEPTIONS_DISABLED
 	catch (OCIO_NAMESPACE::Exception& Exc)
 	{
-		UE_LOG(LogOpenColorIOWrapper, Error, TEXT("Could not create OCIO configuration file for %s. Error message: %s."), FilePath.GetData(), StringCast<TCHAR>(Exc.what()).Get());
+		UE_LOG(LogOpenColorIOWrapper, Error, TEXT("Could not create OCIO configuration file for %s. Error message: %s."), InFilePath.GetData(), StringCast<TCHAR>(Exc.what()).Get());
 	}
 #endif
 #else
@@ -331,29 +350,37 @@ FString FOpenColorIOConfigWrapper::GetCacheID() const
 
 FOpenColorIOProcessorWrapper::FOpenColorIOProcessorWrapper(
 	const FOpenColorIOConfigWrapper* InConfig,
-	FStringView SourceColorSpace,
-	FStringView DestinationColorSpace,
-	EOpenColorIOWorkingColorSpaceTransform InWorkingColorSpaceTransformType,
-	const TMap<FString, FString>& ContextKeyValues)
+	FStringView InSourceColorSpace,
+	FStringView InDestinationColorSpace,
+	const TMap<FString, FString>& InContextKeyValues)
 	: Pimpl(MakePimpl<FOpenColorIOProcessorPimpl, EPimplPtrMode::DeepCopy>())
 	, OwnerConfig(InConfig)
-	, WorkingColorSpaceTransformType(InWorkingColorSpaceTransformType)
+	, WorkingColorSpaceTransformType(EOpenColorIOWorkingColorSpaceTransform::None)
 {
+	if (InSourceColorSpace == OpenColorIOWrapper::GetWorkingColorSpaceName())
+	{
+		WorkingColorSpaceTransformType = EOpenColorIOWorkingColorSpaceTransform::Source;
+	}
+	else if (InDestinationColorSpace == OpenColorIOWrapper::GetWorkingColorSpaceName())
+	{
+		WorkingColorSpaceTransformType = EOpenColorIOWorkingColorSpaceTransform::Destination;
+	}
+
 #if WITH_OCIO
 	if (OwnerConfig != nullptr && OwnerConfig->IsValid())
 	{
 		const OCIO_NAMESPACE::ConstConfigRcPtr& Config = OwnerConfig->Pimpl->Config;
 		 OCIO_NAMESPACE::ContextRcPtr Context = Config->getCurrentContext()->createEditableCopy();
 
-		for (const TPair<FString, FString>& KeyValue : ContextKeyValues)
+		for (const TPair<FString, FString>& KeyValue : InContextKeyValues)
 		{
 			Context->setStringVar(TCHAR_TO_ANSI(*KeyValue.Key), TCHAR_TO_ANSI(*KeyValue.Value));
 		}
 
 		Pimpl->Processor = Config->getProcessor(
 			Context,
-			StringCast<ANSICHAR>(SourceColorSpace.GetData()).Get(),
-			StringCast<ANSICHAR>(DestinationColorSpace.GetData()).Get()
+			StringCast<ANSICHAR>(InSourceColorSpace.GetData()).Get(),
+			StringCast<ANSICHAR>(InDestinationColorSpace.GetData()).Get()
 		);
 	}
 #endif // WITH_OCIO
@@ -361,16 +388,27 @@ FOpenColorIOProcessorWrapper::FOpenColorIOProcessorWrapper(
 
 FOpenColorIOProcessorWrapper::FOpenColorIOProcessorWrapper(
 	const FOpenColorIOConfigWrapper* InConfig,
-	FStringView SourceColorSpace,
-	FStringView Display,
-	FStringView View,
+	FStringView InSourceColorSpace,
+	FStringView InDisplay,
+	FStringView InView,
 	bool bInverseDirection,
-	EOpenColorIOWorkingColorSpaceTransform InWorkingColorSpaceTransformType,
-	const TMap<FString, FString>& ContextKeyValues)
+	const TMap<FString, FString>& InContextKeyValues)
 	: Pimpl(MakePimpl<FOpenColorIOProcessorPimpl, EPimplPtrMode::DeepCopy>())
 	, OwnerConfig(InConfig)
-	, WorkingColorSpaceTransformType(InWorkingColorSpaceTransformType)
+	, WorkingColorSpaceTransformType(EOpenColorIOWorkingColorSpaceTransform::None)
 {
+	if (InSourceColorSpace == OpenColorIOWrapper::GetWorkingColorSpaceName())
+	{
+		if (bInverseDirection)
+		{
+			WorkingColorSpaceTransformType = EOpenColorIOWorkingColorSpaceTransform::Destination;
+		}
+		else
+		{
+			WorkingColorSpaceTransformType = EOpenColorIOWorkingColorSpaceTransform::Source;
+		}
+	}
+
 #if WITH_OCIO
 #if !PLATFORM_EXCEPTIONS_DISABLED
 	try
@@ -381,15 +419,15 @@ FOpenColorIOProcessorWrapper::FOpenColorIOProcessorWrapper(
 			const OCIO_NAMESPACE::ConstConfigRcPtr& Config = OwnerConfig->Pimpl->Config;
 			OCIO_NAMESPACE::ContextRcPtr Context = Config->getCurrentContext()->createEditableCopy();
 
-			for (const TPair<FString, FString>& KeyValue : ContextKeyValues)
+			for (const TPair<FString, FString>& KeyValue : InContextKeyValues)
 			{
 				Context->setStringVar(TCHAR_TO_ANSI(*KeyValue.Key), TCHAR_TO_ANSI(*KeyValue.Value));
 			}
 
 			Pimpl->Processor = Config->getProcessor(
-				StringCast<ANSICHAR>(SourceColorSpace.GetData()).Get(),
-				StringCast<ANSICHAR>(Display.GetData()).Get(),
-				StringCast<ANSICHAR>(View.GetData()).Get(),
+				StringCast<ANSICHAR>(InSourceColorSpace.GetData()).Get(),
+				StringCast<ANSICHAR>(InDisplay.GetData()).Get(),
+				StringCast<ANSICHAR>(InView.GetData()).Get(),
 				static_cast<OCIO_NAMESPACE::TransformDirection>(bInverseDirection)
 			);
 		}
@@ -437,6 +475,7 @@ bool FOpenColorIOCPUProcessorWrapper::TransformImage(const FImageView& InOutImag
 			TUniquePtr<PackedImageDesc> ImageDesc = GetImageDesc(InOutImage);
 			if (ImageDesc)
 			{
+				const TUniquePtr<ANSICHAR[]> AnsiWorkingColorSpaceName = OpenColorIOWrapper::MakeAnsiString(OpenColorIOWrapper::GetWorkingColorSpaceName());
 				ConstConfigRcPtr	InterchangeConfig = IOpenColorIOWrapperModule::Get().GetWorkingColorSpaceToInterchangeConfig()->Pimpl->Config;
 				ConstConfigRcPtr Config = ParentProcessor.OwnerConfig->Pimpl->Config;
 
@@ -447,9 +486,9 @@ bool FOpenColorIOCPUProcessorWrapper::TransformImage(const FImageView& InOutImag
 				{
 					ConstProcessorRcPtr	InterchangeProcessor = InterchangeConfig->GetProcessorFromConfigs(
 						InterchangeConfig,
-						StringCast<ANSICHAR>(OpenColorIOWrapper::WorkingColorSpaceName).Get(),
+						AnsiWorkingColorSpaceName.Get(),
 						Config,
-						Config->getCanonicalName(OpenColorIOWrapper::InterchangeName));
+						Config->getCanonicalName(OpenColorIOWrapper::GetInterchangeName()));
 
 					ConstCPUProcessorRcPtr InterchangeCPUProcessor = InterchangeProcessor->getOptimizedCPUProcessor(BitDepth, BitDepth, OPTIMIZATION_DEFAULT);
 					InterchangeCPUProcessor->apply(*ImageDesc);
@@ -464,9 +503,9 @@ bool FOpenColorIOCPUProcessorWrapper::TransformImage(const FImageView& InOutImag
 				{
 					ConstProcessorRcPtr	InterchangeProcessor = InterchangeConfig->GetProcessorFromConfigs(
 						Config,
-						Config->getCanonicalName(OpenColorIOWrapper::InterchangeName),
+						Config->getCanonicalName(OpenColorIOWrapper::GetInterchangeName()),
 						InterchangeConfig,
-						StringCast<ANSICHAR>(OpenColorIOWrapper::WorkingColorSpaceName).Get());
+						AnsiWorkingColorSpaceName.Get());
 
 					ConstCPUProcessorRcPtr InterchangeCPUProcessor = InterchangeProcessor->getOptimizedCPUProcessor(BitDepth, BitDepth, OPTIMIZATION_DEFAULT);
 					InterchangeCPUProcessor->apply(*ImageDesc);
@@ -503,6 +542,7 @@ bool FOpenColorIOCPUProcessorWrapper::TransformImage(const FImageView& SrcImage,
 			TUniquePtr<PackedImageDesc> DestImageDesc = GetImageDesc(DestImage);
 			if (SrcImageDesc && DestImageDesc)
 			{
+				const TUniquePtr<ANSICHAR[]> AnsiWorkingColorSpaceName = OpenColorIOWrapper::MakeAnsiString(OpenColorIOWrapper::GetWorkingColorSpaceName());
 				ConstConfigRcPtr	InterchangeConfig = IOpenColorIOWrapperModule::Get().GetWorkingColorSpaceToInterchangeConfig()->Pimpl->Config;
 				ConstConfigRcPtr	Config = ParentProcessor.OwnerConfig->Pimpl->Config;
 
@@ -514,9 +554,9 @@ bool FOpenColorIOCPUProcessorWrapper::TransformImage(const FImageView& SrcImage,
 				{
 					ConstProcessorRcPtr	InterchangeProcessor = InterchangeConfig->GetProcessorFromConfigs(
 						InterchangeConfig,
-						StringCast<ANSICHAR>(OpenColorIOWrapper::WorkingColorSpaceName).Get(),
+						AnsiWorkingColorSpaceName.Get(),
 						Config,
-						Config->getCanonicalName(OpenColorIOWrapper::InterchangeName));
+						Config->getCanonicalName(OpenColorIOWrapper::GetInterchangeName()));
 
 					ConstCPUProcessorRcPtr InterchangeCPUProcessor = InterchangeProcessor->getOptimizedCPUProcessor(SrcBitDepth, SrcBitDepth, OPTIMIZATION_DEFAULT);
 					InterchangeCPUProcessor->apply(*SrcImageDesc);
@@ -532,9 +572,9 @@ bool FOpenColorIOCPUProcessorWrapper::TransformImage(const FImageView& SrcImage,
 				{
 					ConstProcessorRcPtr	InterchangeProcessor = InterchangeConfig->GetProcessorFromConfigs(
 						Config,
-						Config->getCanonicalName(OpenColorIOWrapper::InterchangeName),
+						Config->getCanonicalName(OpenColorIOWrapper::GetInterchangeName()),
 						InterchangeConfig,
-						StringCast<ANSICHAR>(OpenColorIOWrapper::WorkingColorSpaceName).Get());
+						AnsiWorkingColorSpaceName.Get());
 
 					ConstCPUProcessorRcPtr InterchangeCPUProcessor = InterchangeProcessor->getOptimizedCPUProcessor(DestBitDepth, DestBitDepth, OPTIMIZATION_DEFAULT);
 					InterchangeCPUProcessor->apply(*DestImageDesc);
@@ -571,7 +611,7 @@ FOpenColorIOGPUProcessorWrapper::FOpenColorIOGPUProcessorWrapper(FOpenColorIOPro
 
 		GpuShaderDescRcPtr ShaderDescription = GpuShaderDesc::CreateShaderDesc();
 		ShaderDescription->setLanguage(GPU_LANGUAGE_HLSL_DX11);
-		ShaderDescription->setFunctionName(StringCast<ANSICHAR>(OpenColorIOWrapper::ShaderFunctionName).Get());
+		ShaderDescription->setFunctionName(StringCast<ANSICHAR>(OpenColorIOWrapper::GetShaderFunctionName()).Get());
 		ShaderDescription->setResourcePrefix("Ocio");
 
 		ConstGPUProcessorRcPtr GPUProcessor = nullptr;
diff --git a/Engine/Plugins/Compositing/OpenColorIO/Source/OpenColorIOWrapper/Private/OpenColorIOWrapperModule.cpp b/Engine/Plugins/Compositing/OpenColorIO/Source/OpenColorIOWrapper/Private/OpenColorIOWrapperModule.cpp
--- a/Engine/Plugins/Compositing/OpenColorIO/Source/OpenColorIOWrapper/Private/OpenColorIOWrapperModule.cpp
+++ b/Engine/Plugins/Compositing/OpenColorIO/Source/OpenColorIOWrapper/Private/OpenColorIOWrapperModule.cpp
@@ -17,32 +17,39 @@ public:
 	//~ Begin IModuleInterface interface
 	virtual void StartupModule() override
 	{
-		bInitializedLib = FOpenColorIOLibHandler::Initialize();
+		FOpenColorIOLibHandler::Initialize();
+
+		InterchangeConfig = FOpenColorIOConfigWrapper::CreateWorkingColorSpaceToInterchangeConfig();
 	}
 
 	virtual void ShutdownModule() override
 	{
 		FOpenColorIOLibHandler::Shutdown();
-		bInitializedLib = false;
 	}
 	//~ End IModuleInterface interface
 
-	virtual const FOpenColorIOConfigWrapper* GetWorkingColorSpaceToInterchangeConfig() override
+	//~ Begin IOpenColorIOWrapperModule interface
+	virtual const FOpenColorIOConfigWrapper* GetWorkingColorSpaceToInterchangeConfig() const override
 	{
-		check(bInitializedLib);
+		return InterchangeConfig.Get();
+	}
 
-		if (!InterchangeConfig)
-		{
-			InterchangeConfig = FOpenColorIOConfigWrapper::CreateWorkingColorSpaceToInterchangeConfig();
-		}
+	virtual void LoadGlobalConfig(FStringView InFilePath) override
+	{
+		FOpenColorIOConfigWrapper::FInitializationOptions Opt;
+		Opt.bAddWorkingColorSpace = true;
+		EngineConfig = MakeUnique<FOpenColorIOConfigWrapper>(InFilePath, Opt);
+	}
 
-		return InterchangeConfig.Get();
+	virtual const FOpenColorIOConfigWrapper* GetGlobalConfig() const override
+	{
+		return EngineConfig.Get();
 	}
+	//~ End IOpenColorIOWrapperModule interface
 private:
-
-	bool bInitializedLib = false;
 	
-	TUniquePtr<FOpenColorIOConfigWrapper> InterchangeConfig;
+	TUniquePtr<FOpenColorIOConfigWrapper> InterchangeConfig = nullptr;
+	TUniquePtr<FOpenColorIOConfigWrapper> EngineConfig = nullptr;
 };
 
 IMPLEMENT_MODULE(FOpenColorIOWrapperModule, OpenColorIOWrapper);
diff --git a/Engine/Plugins/Compositing/OpenColorIO/Source/OpenColorIOWrapper/Public/OpenColorIOWrapper.h b/Engine/Plugins/Compositing/OpenColorIO/Source/OpenColorIOWrapper/Public/OpenColorIOWrapper.h
--- a/Engine/Plugins/Compositing/OpenColorIO/Source/OpenColorIOWrapper/Public/OpenColorIOWrapper.h
+++ b/Engine/Plugins/Compositing/OpenColorIO/Source/OpenColorIOWrapper/Public/OpenColorIOWrapper.h
@@ -17,14 +17,17 @@ enum TextureFilter : int;
 
 namespace OpenColorIOWrapper
 {
-	/** OCIO_NAMESPACE::ROLE_INTERCHANGE_SCENE equivalent, since we currently cannot delay-load this definition. */
-	static constexpr ANSICHAR InterchangeName[] = "aces_interchange";
-	
 	/** Color space name of the engine's working color space we insert in OpenColorIO configs. */
-	static constexpr ANSICHAR WorkingColorSpaceName[] = "Working Color Space";
+	constexpr const TCHAR* GetWorkingColorSpaceName()
+	{
+		return TEXT("Working Color Space");
+	}
 
 	/** Default generated shader function name. */
-	static constexpr TCHAR ShaderFunctionName[] = TEXT("OCIOConvert");
+	constexpr const TCHAR* GetShaderFunctionName()
+	{
+		return TEXT("OCIOConvert");
+	}
 
 	/** Default LUT size used in the legacy gpu processor */
 	static constexpr uint32 Legacy3dEdgeLength = 65;
@@ -51,10 +54,10 @@ public:
 	/**
 	* Constructor.
 	* 
-	* @param FilePath Config absolute file path.
+	* @param InFilePath Config absolute file path.
 	* @param InOptions Initialization options.
 	*/
-	FOpenColorIOConfigWrapper(FStringView FilePath, FInitializationOptions InOptions);
+	FOpenColorIOConfigWrapper(FStringView InFilePath, FInitializationOptions InOptions);
 
 	/** Valid when the native config has been successfully created and isn't null. */
 	bool IsValid() const;
@@ -109,37 +112,33 @@ public:
 	* Constructor.
 	*
 	* @param InConfig Owner config.
-	* @param SourceColorSpace Source color space name.
-	* @param DestinationColorSpace Destination color space name.
-	* @param WorkingColorSpaceTransformType (Optional) Type of additional conversion from/to the working color space.
-	* @param ContextKeyValues (Optional) ContextKeyValues Additional context modifiers.
+	* @param InSourceColorSpace Source color space name.
+	* @param InDestinationColorSpace Destination color space name.
+	* @param InContextKeyValues (Optional) Additional context modifiers.
 	*/
 	FOpenColorIOProcessorWrapper(
 		const FOpenColorIOConfigWrapper* InConfig,
-		FStringView SourceColorSpace,
-		FStringView DestinationColorSpace,
-		EOpenColorIOWorkingColorSpaceTransform WorkingColorSpaceTransformType = EOpenColorIOWorkingColorSpaceTransform::None,
-		const TMap<FString, FString>& ContextKeyValues = {});
+		FStringView InSourceColorSpace,
+		FStringView InDestinationColorSpace,
+		const TMap<FString, FString>& InContextKeyValues = {});
 
 	/**
 	* Constructor.
 	*
 	* @param InConfig Owner config.
-	* @param SourceColorSpace Source color space name.
-	* @param Display Display name in display-view transform.
-	* @param View View name in display-view transform.
+	* @param InSourceColorSpace Source color space name.
+	* @param InDisplay Display name in display-view transform.
+	* @param InView View name in display-view transform.
 	* @param bInverseDirection Flag for inverse transform direction.
-	* @param WorkingColorSpaceTransformType (Optional) Type of additional conversion from/to the working color space.
-	* @param ContextKeyValues (Optional) ContextKeyValues Additional context modifiers.
+	* @param InContextKeyValues (Optional) Additional context modifiers.
 	*/
 	FOpenColorIOProcessorWrapper(
 		const FOpenColorIOConfigWrapper* InConfig,
-		FStringView SourceColorSpace,
-		FStringView Display,
-		FStringView View,
+		FStringView InSourceColorSpace,
+		FStringView InDisplay,
+		FStringView InView,
 		bool bInverseDirection = false,
-		EOpenColorIOWorkingColorSpaceTransform WorkingColorSpaceTransformType = EOpenColorIOWorkingColorSpaceTransform::None,
-		const TMap<FString, FString>& ContextKeyValues = {});
+		const TMap<FString, FString>& InContextKeyValues = {});
 
 	/** Valid when the processor has been successfully created and isn't null. */
 	bool IsValid() const;
diff --git a/Engine/Plugins/Compositing/OpenColorIO/Source/OpenColorIOWrapper/Public/OpenColorIOWrapperModule.h b/Engine/Plugins/Compositing/OpenColorIO/Source/OpenColorIOWrapper/Public/OpenColorIOWrapperModule.h
--- a/Engine/Plugins/Compositing/OpenColorIO/Source/OpenColorIOWrapper/Public/OpenColorIOWrapperModule.h
+++ b/Engine/Plugins/Compositing/OpenColorIO/Source/OpenColorIOWrapper/Public/OpenColorIOWrapperModule.h
@@ -41,7 +41,19 @@ public:
 	/**
 	 * Returns a minimal dynamically-created native config for conversions between interchange and working color spaces.
 	 */
-	virtual const FOpenColorIOConfigWrapper* GetWorkingColorSpaceToInterchangeConfig() = 0;
+	virtual const FOpenColorIOConfigWrapper* GetWorkingColorSpaceToInterchangeConfig() const = 0;
+
+	/**
+	* Load a globally-shared config in the module.
+	*
+	* @param InFilePath Config absolute file path.
+	*/
+	virtual void LoadGlobalConfig(FStringView InFilePath) = 0;
+
+	/**
+	* Returns the globally-shared module config if loaded, nullptr otherwise.
+	*/
+	virtual const FOpenColorIOConfigWrapper* GetGlobalConfig() const = 0;
 
 	/** Virtual destructor */
 	virtual ~IOpenColorIOWrapperModule() = default;
