OCIO: Fix ndisplay crash with new wrapper module.

diff --git a/Engine/Plugins/Compositing/OpenColorIO/OpenColorIO.uplugin b/Engine/Plugins/Compositing/OpenColorIO/OpenColorIO.uplugin
--- a/Engine/Plugins/Compositing/OpenColorIO/OpenColorIO.uplugin
+++ b/Engine/Plugins/Compositing/OpenColorIO/OpenColorIO.uplugin
@@ -17,7 +17,7 @@
 	"Modules": [
 		{
 			"Name": "OpenColorIOWrapper",
-			"Type": "Editor",
+			"Type": "Runtime",
 			"LoadingPhase": "PostConfigInit"
 		},
 		{
diff --git a/Engine/Plugins/Compositing/OpenColorIO/Source/OpenColorIO/Private/OpenColorIOModule.cpp b/Engine/Plugins/Compositing/OpenColorIO/Source/OpenColorIO/Private/OpenColorIOModule.cpp
--- a/Engine/Plugins/Compositing/OpenColorIO/Source/OpenColorIO/Private/OpenColorIOModule.cpp
+++ b/Engine/Plugins/Compositing/OpenColorIO/Source/OpenColorIO/Private/OpenColorIOModule.cpp
@@ -6,6 +6,9 @@
 #include "OpenColorIODisplayManager.h"
 #include "ShaderCore.h"
 
+#if WITH_EDITOR
+#include "OpenColorIOWrapperModule.h"
+#endif
 
 DEFINE_LOG_CATEGORY(LogOpenColorIO);
 
@@ -19,6 +22,11 @@ FOpenColorIOModule::FOpenColorIOModule()
 
 void FOpenColorIOModule::StartupModule()
 {
+#if WITH_EDITOR
+	// Ensure the wrapper/library modules are loaded first.
+	IOpenColorIOWrapperModule::Get();
+#endif
+
 	// Maps virtual shader source directory /Plugin/OpenCVLensDistortion to the plugin's actual Shaders directory.
 	FString PluginShaderDir = FPaths::Combine(IPluginManager::Get().FindPlugin(TEXT("OpenColorIO"))->GetBaseDir(), TEXT("Shaders"));
 	AddShaderSourceDirectoryMapping(TEXT("/Plugin/OpenColorIO"), PluginShaderDir);
diff --git a/Engine/Plugins/Compositing/OpenColorIO/Source/OpenColorIOWrapper/OpenColorIOWrapper.Build.cs b/Engine/Plugins/Compositing/OpenColorIO/Source/OpenColorIOWrapper/OpenColorIOWrapper.Build.cs
--- a/Engine/Plugins/Compositing/OpenColorIO/Source/OpenColorIOWrapper/OpenColorIOWrapper.Build.cs
+++ b/Engine/Plugins/Compositing/OpenColorIO/Source/OpenColorIOWrapper/OpenColorIOWrapper.Build.cs
@@ -16,9 +16,17 @@ namespace UnrealBuildTool.Rules
 			PrivateDependencyModuleNames.AddRange(new string[]
 			{
 				"ColorManagement",
-				"OpenColorIOLib" // The OpenColorIO third-party library will only be available on desktop editor.
 			});
 
+			if (Target.bBuildEditor)
+			{
+				PrivateDependencyModuleNames.Add("OpenColorIOLib");
+			}
+			else
+			{
+				PrivateDefinitions.Add("WITH_OCIO=0");
+			}
+
 			PublicDependencyModuleNames.AddRange(new string[]
 			{
 				"Core",
