OCIO: Missed try-catch block on new processor wrapper.

diff --git a/Engine/Plugins/Compositing/OpenColorIO/Source/OpenColorIOWrapper/Private/OpenColorIOWrapper.cpp b/Engine/Plugins/Compositing/OpenColorIO/Source/OpenColorIOWrapper/Private/OpenColorIOWrapper.cpp
--- a/Engine/Plugins/Compositing/OpenColorIO/Source/OpenColorIOWrapper/Private/OpenColorIOWrapper.cpp
+++ b/Engine/Plugins/Compositing/OpenColorIO/Source/OpenColorIOWrapper/Private/OpenColorIOWrapper.cpp
@@ -367,22 +367,33 @@ FOpenColorIOProcessorWrapper::FOpenColorIOProcessorWrapper(
 	}
 
 #if WITH_OCIO
-	if (OwnerConfig != nullptr && OwnerConfig->IsValid())
+#if !PLATFORM_EXCEPTIONS_DISABLED
+	try
+#endif
 	{
-		const OCIO_NAMESPACE::ConstConfigRcPtr& Config = OwnerConfig->Pimpl->Config;
-		 OCIO_NAMESPACE::ContextRcPtr Context = Config->getCurrentContext()->createEditableCopy();
-
-		for (const TPair<FString, FString>& KeyValue : InContextKeyValues)
+		if (OwnerConfig != nullptr && OwnerConfig->IsValid())
 		{
-			Context->setStringVar(TCHAR_TO_ANSI(*KeyValue.Key), TCHAR_TO_ANSI(*KeyValue.Value));
-		}
+			const OCIO_NAMESPACE::ConstConfigRcPtr& Config = OwnerConfig->Pimpl->Config;
+			 OCIO_NAMESPACE::ContextRcPtr Context = Config->getCurrentContext()->createEditableCopy();
 
-		Pimpl->Processor = Config->getProcessor(
-			Context,
-			StringCast<ANSICHAR>(InSourceColorSpace.GetData()).Get(),
-			StringCast<ANSICHAR>(InDestinationColorSpace.GetData()).Get()
-		);
+			for (const TPair<FString, FString>& KeyValue : InContextKeyValues)
+			{
+				Context->setStringVar(TCHAR_TO_ANSI(*KeyValue.Key), TCHAR_TO_ANSI(*KeyValue.Value));
+			}
+
+			Pimpl->Processor = Config->getProcessor(
+				Context,
+				StringCast<ANSICHAR>(InSourceColorSpace.GetData()).Get(),
+				StringCast<ANSICHAR>(InDestinationColorSpace.GetData()).Get()
+			);
+		}
 	}
+#if !PLATFORM_EXCEPTIONS_DISABLED
+	catch (OCIO_NAMESPACE::Exception& Exc)
+	{
+		UE_LOG(LogOpenColorIOWrapper, Log, TEXT("Failed to create processor. Error message: %s"), StringCast<TCHAR>(Exc.what()).Get());
+	}
+#endif
 #endif // WITH_OCIO
 }
 
@@ -435,7 +446,7 @@ FOpenColorIOProcessorWrapper::FOpenColorIOProcessorWrapper(
 #if !PLATFORM_EXCEPTIONS_DISABLED
 	catch (OCIO_NAMESPACE::Exception& Exc)
 	{
-		UE_LOG(LogOpenColorIOWrapper, Log, TEXT("Failed to create processor image. Error message: %s"), StringCast<TCHAR>(Exc.what()).Get());
+		UE_LOG(LogOpenColorIOWrapper, Log, TEXT("Failed to create processor. Error message: %s"), StringCast<TCHAR>(Exc.what()).Get());
 	}
 #endif
 #endif // WITH_OCIO
