OCIO: Prevent crash on loading invalid transforms.

(Further investigation will be needed as to why now missing transforms are still requested in postload.)

#rnx

diff --git a/Engine/Plugins/Compositing/OpenColorIO/Source/OpenColorIO/Private/OpenColorIOColorTransform.cpp b/Engine/Plugins/Compositing/OpenColorIO/Source/OpenColorIO/Private/OpenColorIOColorTransform.cpp
--- a/Engine/Plugins/Compositing/OpenColorIO/Source/OpenColorIO/Private/OpenColorIOColorTransform.cpp
+++ b/Engine/Plugins/Compositing/OpenColorIO/Source/OpenColorIO/Private/OpenColorIOColorTransform.cpp
@@ -25,6 +25,21 @@
 
 namespace
 {
+	// Check whether the configuration asset has correctly created its native OpenColorIO config object.
+	bool HasValidConfiguration(const TObjectPtr<UOpenColorIOConfiguration>& InConfiguration)
+	{
+		if (IsValid(InConfiguration))
+		{
+			const FOpenColorIOConfigWrapper* ConfigWrapper = InConfiguration->GetConfigWrapper();
+			if (ConfigWrapper != nullptr)
+			{
+				return ConfigWrapper->IsValid();
+			}
+		}
+
+		return false;
+	}
+
 	// Note: Since OpenColorIO caches processors automatically, we can recreate them without significant additional costs.
 	FOpenColorIOProcessorWrapper GetTransformProcessor(const UOpenColorIOColorTransform* InTransform)
 	{
@@ -263,56 +278,59 @@ void UOpenColorIOColorTransform::CacheResourceTextures()
 #if WITH_EDITOR
 	if (Textures.IsEmpty())
 	{
-		const FOpenColorIOProcessorWrapper TransformProcessor = GetTransformProcessor(this);
-		if (TransformProcessor.IsValid())
+		if (HasValidConfiguration(ConfigurationOwner))
 		{
-			const FOpenColorIOGPUProcessorWrapper GPUProcessor = FOpenColorIOGPUProcessorWrapper(TransformProcessor);
-			const FString ProcessorID = GPUProcessor.GetCacheID();
+			const FOpenColorIOProcessorWrapper TransformProcessor = GetTransformProcessor(this);
+			if (TransformProcessor.IsValid())
+			{
+				const FOpenColorIOGPUProcessorWrapper GPUProcessor = FOpenColorIOGPUProcessorWrapper(TransformProcessor);
+				const FString ProcessorID = GPUProcessor.GetCacheID();
 
-			//In editor, it will use what's on DDC if there's something corresponding to the actual data or use that raw data
-			//that OCIO library has on board. The textures will be serialized only when cooking.
+				//In editor, it will use what's on DDC if there's something corresponding to the actual data or use that raw data
+				//that OCIO library has on board. The textures will be serialized only when cooking.
 
-			// Process 3D luts
-			for (uint32 Index = 0; Index < GPUProcessor.GetNum3DTextures(); ++Index)
-			{
-				TextureFilter Filter;
-				FName TextureName;
-				uint32 EdgeLength = 0;
-				const float* TextureValues = 0x0;
+				// Process 3D luts
+				for (uint32 Index = 0; Index < GPUProcessor.GetNum3DTextures(); ++Index)
+				{
+					TextureFilter Filter;
+					FName TextureName;
+					uint32 EdgeLength = 0;
+					const float* TextureValues = 0x0;
 
-				bool bSuccess = GPUProcessor.Get3DTexture(Index, TextureName, EdgeLength, Filter, TextureValues);
-				check(bSuccess);
+					bool bSuccess = GPUProcessor.Get3DTexture(Index, TextureName, EdgeLength, Filter, TextureValues);
+					check(bSuccess);
 
-				TObjectPtr<UTexture> Result = CreateTexture3DLUT(ProcessorID, TextureName, EdgeLength, Filter, TextureValues);
+					TObjectPtr<UTexture> Result = CreateTexture3DLUT(ProcessorID, TextureName, EdgeLength, Filter, TextureValues);
 
-				const int32 SlotIndex = TextureName.GetNumber() - 1; // Rely on FName's index number extraction for convenience
-				Textures.Add(SlotIndex, MoveTemp(Result));
-			}
+					const int32 SlotIndex = TextureName.GetNumber() - 1; // Rely on FName's index number extraction for convenience
+					Textures.Add(SlotIndex, MoveTemp(Result));
+				}
 
-			// Process 1D luts
-			for (uint32 Index = 0; Index < GPUProcessor.GetNumTextures(); ++Index)
-			{
-				FName TextureName;
-				uint32 TextureWidth = 0;
-				uint32 TextureHeight = 0;
-				TextureFilter Filter;
-				bool bRedChannelOnly = false;
-				const float* TextureValues = 0x0;
+				// Process 1D luts
+				for (uint32 Index = 0; Index < GPUProcessor.GetNumTextures(); ++Index)
+				{
+					FName TextureName;
+					uint32 TextureWidth = 0;
+					uint32 TextureHeight = 0;
+					TextureFilter Filter;
+					bool bRedChannelOnly = false;
+					const float* TextureValues = 0x0;
 
-				bool bSuccess = GPUProcessor.GetTexture(Index, TextureName, TextureWidth, TextureHeight, Filter, bRedChannelOnly, TextureValues);
-				checkf(bSuccess, TEXT("Failed to read OCIO 1D LUT data."));
+					bool bSuccess = GPUProcessor.GetTexture(Index, TextureName, TextureWidth, TextureHeight, Filter, bRedChannelOnly, TextureValues);
+					checkf(bSuccess, TEXT("Failed to read OCIO 1D LUT data."));
 
-				TObjectPtr<UTexture> Result = CreateTexture1DLUT(ProcessorID, TextureName, TextureWidth, TextureHeight, Filter, bRedChannelOnly, TextureValues);
+					TObjectPtr<UTexture> Result = CreateTexture1DLUT(ProcessorID, TextureName, TextureWidth, TextureHeight, Filter, bRedChannelOnly, TextureValues);
 
-				const int32 SlotIndex = TextureName.GetNumber() - 1; // Rely on FName's index number extraction for convenience
-				Textures.Add(SlotIndex, MoveTemp(Result));
-			}
+					const int32 SlotIndex = TextureName.GetNumber() - 1; // Rely on FName's index number extraction for convenience
+					Textures.Add(SlotIndex, MoveTemp(Result));
+				}
 
-			ensureAlwaysMsgf(Textures.Num() <= (int32)OpenColorIOShader::MaximumTextureSlots, TEXT("Color transform %s exceeds our current limit of %u texture slots. Use the legacy processor instead."), *GetTransformFriendlyName(), OpenColorIOShader::MaximumTextureSlots);
+				ensureAlwaysMsgf(Textures.Num() <= (int32)OpenColorIOShader::MaximumTextureSlots, TEXT("Color transform %s exceeds our current limit of %u texture slots. Use the legacy processor instead."), *GetTransformFriendlyName(), OpenColorIOShader::MaximumTextureSlots);
+			}
 		}
 		else
 		{
-			UE_LOG(LogOpenColorIO, Error, TEXT("Failed to cache texture resource(s) for color transform %s. Configuration file [%s] was invalid."), *GetTransformFriendlyName(), *ConfigurationOwner->ConfigurationFile.FilePath);
+			UE_LOG(LogOpenColorIO, Error, TEXT("Failed to cache texture resource(s) for color transform %s. Configuration file [%s] was invalid."), *GetTransformFriendlyName(), ConfigurationOwner ? *ConfigurationOwner->ConfigurationFile.FilePath : TEXT("Missing"));
 		}
 	}
 #endif // WITH_EDITOR
@@ -474,12 +492,19 @@ bool UOpenColorIOColorTransform::IsTransform(const FString& InSourceColorSpace,
 #if WITH_EDITOR
 bool UOpenColorIOColorTransform::EditorTransformImage(const FImageView& InOutImage) const
 {
-	const FOpenColorIOProcessorWrapper TransformProcessor = GetTransformProcessor(this);
-	if (TransformProcessor.IsValid())
+	if (HasValidConfiguration(ConfigurationOwner))
 	{
-		const FOpenColorIOCPUProcessorWrapper CPUProcessor = FOpenColorIOCPUProcessorWrapper(TransformProcessor);
+		const FOpenColorIOProcessorWrapper TransformProcessor = GetTransformProcessor(this);
+		if (TransformProcessor.IsValid())
+		{
+			const FOpenColorIOCPUProcessorWrapper CPUProcessor = FOpenColorIOCPUProcessorWrapper(TransformProcessor);
 
-		return CPUProcessor.TransformImage(InOutImage);
+			return CPUProcessor.TransformImage(InOutImage);
+		}
+	}
+	else
+	{
+		UE_LOG(LogOpenColorIO, Error, TEXT("Failed to transform image for color transform %s. Configuration file was invalid."), *GetTransformFriendlyName());
 	}
 
 	return false;
@@ -487,12 +512,19 @@ bool UOpenColorIOColorTransform::EditorTransformImage(const FImageView& InOutIma
 
 bool UOpenColorIOColorTransform::EditorTransformImage(const FImageView& SrcImage, const FImageView& DestImage) const
 {
-	const FOpenColorIOProcessorWrapper TransformProcessor = GetTransformProcessor(this);
-	if (TransformProcessor.IsValid())
+	if (HasValidConfiguration(ConfigurationOwner))
 	{
-		const FOpenColorIOCPUProcessorWrapper CPUProcessor = FOpenColorIOCPUProcessorWrapper(TransformProcessor);
+		const FOpenColorIOProcessorWrapper TransformProcessor = GetTransformProcessor(this);
+		if (TransformProcessor.IsValid())
+		{
+			const FOpenColorIOCPUProcessorWrapper CPUProcessor = FOpenColorIOCPUProcessorWrapper(TransformProcessor);
 
-		return CPUProcessor.TransformImage(SrcImage, DestImage);
+			return CPUProcessor.TransformImage(SrcImage, DestImage);
+		}
+	}
+	else
+	{
+		UE_LOG(LogOpenColorIO, Error, TEXT("Failed to transform image for color transform %s. Configuration file was invalid."), *GetTransformFriendlyName());
 	}
 
 	return false;
@@ -562,7 +594,7 @@ bool UOpenColorIOColorTransform::GenerateColorTransformData(const FString& InSou
 	return false;
 }
 
-FString UOpenColorIOColorTransform::GetTransformFriendlyName()
+FString UOpenColorIOColorTransform::GetTransformFriendlyName() const
 {
 	if (bIsDisplayViewType)
 	{
@@ -586,17 +618,25 @@ FString UOpenColorIOColorTransform::GetTransformFriendlyName()
 bool UOpenColorIOColorTransform::UpdateShaderInfo(FString& OutShaderCodeHash, FString& OutShaderCode, FString& OutRawConfigHash)
 {
 #if WITH_EDITOR
-	const FOpenColorIOProcessorWrapper Processor = GetTransformProcessor(this);
-	const FOpenColorIOGPUProcessorWrapper GPUProcessor = FOpenColorIOGPUProcessorWrapper(Processor);
-
-	if (GPUProcessor.IsValid())
+	if (HasValidConfiguration(ConfigurationOwner))
 	{
-		OutRawConfigHash = ConfigurationOwner->GetConfigWrapper()->GetCacheID();
+		const FOpenColorIOProcessorWrapper Processor = GetTransformProcessor(this);
+		if (Processor.IsValid())
+		{
+			const FOpenColorIOGPUProcessorWrapper GPUProcessor = FOpenColorIOGPUProcessorWrapper(Processor);
+			if (GPUProcessor.IsValid())
+			{
+				OutRawConfigHash = ConfigurationOwner->GetConfigWrapper()->GetCacheID();
 
-		return GPUProcessor.GetShader(OutShaderCode, OutShaderCodeHash);
+				return GPUProcessor.GetShader(OutShaderCode, OutShaderCodeHash);
+			}
+		}
 	}
-	
-	UE_LOG(LogOpenColorIO, Error, TEXT("Failed to fetch shader info for color transform %s. Configuration file [%s] was invalid."), *GetTransformFriendlyName(), *ConfigurationOwner->ConfigurationFile.FilePath);
+	else
+	{
+		UE_LOG(LogOpenColorIO, Error, TEXT("Failed to fetch shader info for color transform %s. Configuration file [%s] was invalid."), *GetTransformFriendlyName(), ConfigurationOwner ? *ConfigurationOwner->ConfigurationFile.FilePath : TEXT("Missing"));
+	}
+
 	return false;
 #else
 	return true; //When not in editor, shaders have been cooked so we're not relying on the library data anymore.
diff --git a/Engine/Plugins/Compositing/OpenColorIO/Source/OpenColorIO/Public/OpenColorIOColorTransform.h b/Engine/Plugins/Compositing/OpenColorIO/Source/OpenColorIO/Public/OpenColorIOColorTransform.h
--- a/Engine/Plugins/Compositing/OpenColorIO/Source/OpenColorIO/Public/OpenColorIOColorTransform.h
+++ b/Engine/Plugins/Compositing/OpenColorIO/Source/OpenColorIO/Public/OpenColorIOColorTransform.h
@@ -128,7 +128,7 @@ protected:
 	/**
 	 * Helper function returning the color space transform name based on source and destination color spaces.
 	 */
-	FString GetTransformFriendlyName();
+	FString GetTransformFriendlyName() const;
 
 	/**
 	 * Fetch shader code and hash from the OCIO library
diff --git a/Engine/Plugins/Compositing/OpenColorIO/Source/OpenColorIOWrapper/Private/OpenColorIOWrapper.cpp b/Engine/Plugins/Compositing/OpenColorIO/Source/OpenColorIOWrapper/Private/OpenColorIOWrapper.cpp
--- a/Engine/Plugins/Compositing/OpenColorIO/Source/OpenColorIOWrapper/Private/OpenColorIOWrapper.cpp
+++ b/Engine/Plugins/Compositing/OpenColorIO/Source/OpenColorIOWrapper/Private/OpenColorIOWrapper.cpp
@@ -391,7 +391,7 @@ FOpenColorIOProcessorWrapper::FOpenColorIOProcessorWrapper(
 #if !PLATFORM_EXCEPTIONS_DISABLED
 	catch (OCIO_NAMESPACE::Exception& Exc)
 	{
-		UE_LOG(LogOpenColorIOWrapper, Log, TEXT("Failed to create processor. Error message: %s"), StringCast<TCHAR>(Exc.what()).Get());
+		UE_LOG(LogOpenColorIOWrapper, Log, TEXT("Failed to create processor for [%s, %s]. Error message: %s"), InSourceColorSpace.GetData(), InDestinationColorSpace.GetData(), StringCast<TCHAR>(Exc.what()).Get());
 	}
 #endif
 #endif // WITH_OCIO
@@ -446,7 +446,7 @@ FOpenColorIOProcessorWrapper::FOpenColorIOProcessorWrapper(
 #if !PLATFORM_EXCEPTIONS_DISABLED
 	catch (OCIO_NAMESPACE::Exception& Exc)
 	{
-		UE_LOG(LogOpenColorIOWrapper, Log, TEXT("Failed to create processor. Error message: %s"), StringCast<TCHAR>(Exc.what()).Get());
+		UE_LOG(LogOpenColorIOWrapper, Log, TEXT("Failed to create processor for [%s, %s, %s, %s]. Error message: %s"), InSourceColorSpace.GetData(), InDisplay.GetData(), InView.GetData(), (bInverseDirection ? TEXT("Inverse") : TEXT("Forward")), StringCast<TCHAR>(Exc.what()).Get());
 	}
 #endif
 #endif // WITH_OCIO
@@ -481,7 +481,7 @@ bool FOpenColorIOCPUProcessorWrapper::TransformImage(const FImageView& InOutImag
 	{
 		using namespace OCIO_NAMESPACE;
 
-		if (ParentProcessor.IsValid())
+		if (IsValid())
 		{
 			TUniquePtr<PackedImageDesc> ImageDesc = GetImageDesc(InOutImage);
 			if (ImageDesc)
@@ -547,7 +547,7 @@ bool FOpenColorIOCPUProcessorWrapper::TransformImage(const FImageView& SrcImage,
 	{
 		using namespace OCIO_NAMESPACE;
 
-		if (ParentProcessor.IsValid())
+		if (IsValid())
 		{
 			TUniquePtr<PackedImageDesc> SrcImageDesc = GetImageDesc(SrcImage);
 			TUniquePtr<PackedImageDesc> DestImageDesc = GetImageDesc(DestImage);
@@ -613,39 +613,42 @@ FOpenColorIOGPUProcessorWrapper::FOpenColorIOGPUProcessorWrapper(FOpenColorIOPro
 	, GPUPimpl(MakePimpl<FOpenColorIOGPUProcessorPimpl, EPimplPtrMode::DeepCopy>())
 {
 #if WITH_OCIO
-	TRACE_CPUPROFILER_EVENT_SCOPE(FOpenColorIOGPUProcessorWrapper::FOpenColorIOGPUProcessorWrapper)
+	TRACE_CPUPROFILER_EVENT_SCOPE(FOpenColorIOGPUProcessorWrapper::FOpenColorIOGPUProcessorWrapper)	
 #if !PLATFORM_EXCEPTIONS_DISABLED
 	try
 #endif
 	{
 		using namespace OCIO_NAMESPACE;
 
-		GpuShaderDescRcPtr ShaderDescription = GpuShaderDesc::CreateShaderDesc();
-		ShaderDescription->setLanguage(GPU_LANGUAGE_HLSL_DX11);
-		ShaderDescription->setFunctionName(StringCast<ANSICHAR>(OpenColorIOWrapper::GetShaderFunctionName()).Get());
-		ShaderDescription->setResourcePrefix("Ocio");
+		if (ParentProcessor.IsValid())
+		{
+			GpuShaderDescRcPtr ShaderDescription = GpuShaderDesc::CreateShaderDesc();
+			ShaderDescription->setLanguage(GPU_LANGUAGE_HLSL_DX11);
+			ShaderDescription->setFunctionName(StringCast<ANSICHAR>(OpenColorIOWrapper::GetShaderFunctionName()).Get());
+			ShaderDescription->setResourcePrefix("Ocio");
 
-		ConstGPUProcessorRcPtr GPUProcessor = nullptr;
-		OptimizationFlags OptFlags = FOpenColorIOProcessorPimpl::GetOptimizationFlags();
+			ConstGPUProcessorRcPtr GPUProcessor = nullptr;
+			OptimizationFlags OptFlags = FOpenColorIOProcessorPimpl::GetOptimizationFlags();
 
-		if (InShaderParams.bIsLegacy)
-		{
-			unsigned int EdgeLength = static_cast<unsigned int>(OpenColorIOWrapper::Legacy3dEdgeLength);
-			GPUProcessor = ParentProcessor.Pimpl->Processor->getOptimizedLegacyGPUProcessor(OptFlags, EdgeLength);
-		}
-		else
-		{
-			GPUProcessor = ParentProcessor.Pimpl->Processor->getOptimizedGPUProcessor(FOpenColorIOProcessorPimpl::GetOptimizationFlags());
-		}
-		GPUProcessor->extractGpuShaderInfo(ShaderDescription);
+			if (InShaderParams.bIsLegacy)
+			{
+				unsigned int EdgeLength = static_cast<unsigned int>(OpenColorIOWrapper::Legacy3dEdgeLength);
+				GPUProcessor = ParentProcessor.Pimpl->Processor->getOptimizedLegacyGPUProcessor(OptFlags, EdgeLength);
+			}
+			else
+			{
+				GPUProcessor = ParentProcessor.Pimpl->Processor->getOptimizedGPUProcessor(FOpenColorIOProcessorPimpl::GetOptimizationFlags());
+			}
+			GPUProcessor->extractGpuShaderInfo(ShaderDescription);
 
-		GPUPimpl->Processor = GPUProcessor;
-		GPUPimpl->ShaderDescription = ShaderDescription;
+			GPUPimpl->Processor = GPUProcessor;
+			GPUPimpl->ShaderDescription = ShaderDescription;
+		}
 	}
 #if !PLATFORM_EXCEPTIONS_DISABLED
 	catch (OCIO_NAMESPACE::Exception& Exc)
 	{
-		UE_LOG(LogOpenColorIOWrapper, Error, TEXT("Failed to fetch shader info for color transform. Error message: %s"), StringCast<TCHAR>(Exc.what()).Get());
+		UE_LOG(LogOpenColorIOWrapper, Log, TEXT("Failed to fetch shader info for color transform. Error message: %s"), StringCast<TCHAR>(Exc.what()).Get());
 	}
 #endif
 #endif // WITH_OCIO
