OCIO: Exposing cpu transform to blueprints.

diff --git a/Engine/Plugins/Compositing/OpenColorIO/Source/OpenColorIOEditor/OpenColorIOEditor.Build.cs b/Engine/Plugins/Compositing/OpenColorIO/Source/OpenColorIOEditor/OpenColorIOEditor.Build.cs
--- a/Engine/Plugins/Compositing/OpenColorIO/Source/OpenColorIOEditor/OpenColorIOEditor.Build.cs
+++ b/Engine/Plugins/Compositing/OpenColorIO/Source/OpenColorIOEditor/OpenColorIOEditor.Build.cs
@@ -14,6 +14,7 @@ namespace UnrealBuildTool.Rules
 					"DesktopWidgets",
 					"EditorFramework",
 					"Engine",
+					"ImageCore",
 					"LevelEditor",
 					"OpenColorIOWrapper",
 					"OpenColorIO",
diff --git a/Engine/Plugins/Compositing/OpenColorIO/Source/OpenColorIOEditor/Private/OpenColorIOEditorBlueprintLibrary.cpp b/Engine/Plugins/Compositing/OpenColorIO/Source/OpenColorIOEditor/Private/OpenColorIOEditorBlueprintLibrary.cpp
--- a/Engine/Plugins/Compositing/OpenColorIO/Source/OpenColorIOEditor/Private/OpenColorIOEditorBlueprintLibrary.cpp
+++ b/Engine/Plugins/Compositing/OpenColorIO/Source/OpenColorIOEditor/Private/OpenColorIOEditorBlueprintLibrary.cpp
@@ -4,8 +4,11 @@
 
 #include "Modules/ModuleManager.h"
 
+#include "Engine/Texture2D.h"
 #include "OpenColorIOEditorModule.h"
 #include "OpenColorIOColorSpace.h"
+#include "OpenColorIOConfiguration.h"
+#include "TextureCompiler.h"
 
 #include UE_INLINE_GENERATED_CPP_BY_NAME(OpenColorIOEditorBlueprintLibrary)
 
@@ -16,3 +19,32 @@ void UOpenColorIOEditorBlueprintLibrary::SetActiveViewportConfiguration(const FO
 
 	OpenColorIOEditorModule.SetActiveViewportConfiguration(InConfiguration);
 }
+
+bool UOpenColorIOEditorBlueprintLibrary::ApplyTextureColorTransform(const FOpenColorIOColorConversionSettings& ConversionSettings, UTexture* InOutTexture, bool bSynchronous)
+{	
+	if (IsValid(InOutTexture) && IsValid(ConversionSettings.ConfigurationSource))
+	{
+		FImage ImageMip0;
+		if (InOutTexture->Source.GetMipImage(ImageMip0, 0))
+		{
+			bool bTransformSucceeded = ConversionSettings.ConfigurationSource->EditorTransformImage(ConversionSettings, ImageMip0);
+			if (bTransformSucceeded)
+			{
+				void* TargetData = InOutTexture->Source.LockMip(0);
+				FMemory::Memcpy(TargetData, ImageMip0.RawData.GetData(), ImageMip0.GetImageSizeBytes());
+				InOutTexture->Source.UnlockMip(0);
+				InOutTexture->Modify();
+				InOutTexture->UpdateResource();
+
+				if (bSynchronous)
+				{
+					FTextureCompilingManager::Get().FinishCompilation({ InOutTexture });
+				}
+				
+				return true;
+			}
+		}
+	}
+
+	return false;
+}
diff --git a/Engine/Plugins/Compositing/OpenColorIO/Source/OpenColorIOEditor/Public/OpenColorIOEditorBlueprintLibrary.h b/Engine/Plugins/Compositing/OpenColorIO/Source/OpenColorIOEditor/Public/OpenColorIOEditorBlueprintLibrary.h
--- a/Engine/Plugins/Compositing/OpenColorIO/Source/OpenColorIOEditor/Public/OpenColorIOEditorBlueprintLibrary.h
+++ b/Engine/Plugins/Compositing/OpenColorIO/Source/OpenColorIOEditor/Public/OpenColorIOEditorBlueprintLibrary.h
@@ -21,4 +21,16 @@ public:
 	 */
 	UFUNCTION(BlueprintCallable, Category = "OpenColorIO")
 	static OPENCOLORIOEDITOR_API void SetActiveViewportConfiguration(const FOpenColorIODisplayConfiguration& InConfiguration);
+
+
+	/**
+	 * Apply a color transform to a texture asset.
+	 *
+	 * @param ConversionSettings Color transformation settings.
+	 * @param InOutTexture Texture object to transform.
+	 * @param bSynchronous Whether the texture transform is a nlocking operation.
+	 * @return true upon success.
+	 */
+	UFUNCTION(BlueprintCallable, Category = "OpenColorIO")
+	static OPENCOLORIOEDITOR_API bool ApplyTextureColorTransform(const FOpenColorIOColorConversionSettings& ConversionSettings, UTexture* InOutTexture, bool bSynchronous = false);
 };
