OCIO: Rename just-added editor BP function to "ApplyColorSpaceTransformToTexture" for consistency with runtime BP function.

diff --git a/Engine/Plugins/Compositing/OpenColorIO/Source/OpenColorIOEditor/Private/OpenColorIOEditorBlueprintLibrary.cpp b/Engine/Plugins/Compositing/OpenColorIO/Source/OpenColorIOEditor/Private/OpenColorIOEditorBlueprintLibrary.cpp
--- a/Engine/Plugins/Compositing/OpenColorIO/Source/OpenColorIOEditor/Private/OpenColorIOEditorBlueprintLibrary.cpp
+++ b/Engine/Plugins/Compositing/OpenColorIO/Source/OpenColorIOEditor/Private/OpenColorIOEditorBlueprintLibrary.cpp
@@ -20,7 +20,7 @@ void UOpenColorIOEditorBlueprintLibrary::SetActiveViewportConfiguration(const FO
 	OpenColorIOEditorModule.SetActiveViewportConfiguration(InConfiguration);
 }
 
-bool UOpenColorIOEditorBlueprintLibrary::ApplyTextureColorTransform(const FOpenColorIOColorConversionSettings& ConversionSettings, UTexture* InOutTexture, bool bSynchronous)
+bool UOpenColorIOEditorBlueprintLibrary::ApplyColorSpaceTransformToTexture(const FOpenColorIOColorConversionSettings& ConversionSettings, UTexture* InOutTexture, bool bSynchronous)
 {	
 	if (IsValid(InOutTexture) && IsValid(ConversionSettings.ConfigurationSource))
 	{
diff --git a/Engine/Plugins/Compositing/OpenColorIO/Source/OpenColorIOEditor/Public/OpenColorIOEditorBlueprintLibrary.h b/Engine/Plugins/Compositing/OpenColorIO/Source/OpenColorIOEditor/Public/OpenColorIOEditorBlueprintLibrary.h
--- a/Engine/Plugins/Compositing/OpenColorIO/Source/OpenColorIOEditor/Public/OpenColorIOEditorBlueprintLibrary.h
+++ b/Engine/Plugins/Compositing/OpenColorIO/Source/OpenColorIOEditor/Public/OpenColorIOEditorBlueprintLibrary.h
@@ -24,7 +24,7 @@ public:
 
 
 	/**
-	 * Apply a color transform to a texture asset.
+	 * Apply a color space transform to a texture asset.
 	 *
 	 * @param ConversionSettings Color transformation settings.
 	 * @param InOutTexture Texture object to transform.
@@ -32,5 +32,5 @@ public:
 	 * @return true upon success.
 	 */
 	UFUNCTION(BlueprintCallable, Category = "OpenColorIO")
-	static OPENCOLORIOEDITOR_API bool ApplyTextureColorTransform(const FOpenColorIOColorConversionSettings& ConversionSettings, UTexture* InOutTexture, bool bSynchronous = false);
+	static OPENCOLORIOEDITOR_API bool ApplyColorSpaceTransformToTexture(const FOpenColorIOColorConversionSettings& ConversionSettings, UTexture* InOutTexture, bool bSynchronous = false);
 };
