OCIO: Fixes for deactivated transforms.

* Prevent resource initialization for loaded but missing transforms (due to previous logic).
* Reload all color spaces on config file change.
* Deprecate redundant/circular object pointer to transform outer, including initialization functions.
* Validate color spaces immediately on UI creation, so that <Invalid> always appears if a selected color space is now missing.
* Fixed disappearing display-view selections in v1 configs.

diff --git a/Engine/Plugins/Compositing/OpenColorIO/Source/OpenColorIO/Private/OpenColorIOColorTransform.cpp b/Engine/Plugins/Compositing/OpenColorIO/Source/OpenColorIO/Private/OpenColorIOColorTransform.cpp
--- a/Engine/Plugins/Compositing/OpenColorIO/Source/OpenColorIO/Private/OpenColorIOColorTransform.cpp
+++ b/Engine/Plugins/Compositing/OpenColorIO/Source/OpenColorIO/Private/OpenColorIOColorTransform.cpp
@@ -25,34 +25,28 @@
 
 namespace
 {
-	// Check whether the configuration asset has correctly created its native OpenColorIO config object.
-	bool HasValidConfiguration(const TObjectPtr<UOpenColorIOConfiguration>& InConfiguration)
+	// Returns the (native) config wrapper from the configuration object.
+	FOpenColorIOConfigWrapper* GetTransformConfigWrapper(const UOpenColorIOConfiguration* InConfigurationOwner)
 	{
-		if (IsValid(InConfiguration))
+		if (IsValid(InConfigurationOwner))
 		{
-			const FOpenColorIOConfigWrapper* ConfigWrapper = InConfiguration->GetConfigWrapper();
-			if (ConfigWrapper != nullptr)
-			{
-				return ConfigWrapper->IsValid();
-			}
+			return InConfigurationOwner->GetConfigWrapper();
 		}
 
-		return false;
+		return nullptr;
 	}
 
 	// Note: Since OpenColorIO caches processors automatically, we can recreate them without significant additional costs.
-	FOpenColorIOProcessorWrapper GetTransformProcessor(const UOpenColorIOColorTransform* InTransform)
+	FOpenColorIOProcessorWrapper GetTransformProcessor(const UOpenColorIOColorTransform* InTransform, const FOpenColorIOConfigWrapper* InConfigWrapper)
 	{
 		check(InTransform);
-		check(InTransform->ConfigurationOwner);
-
-		const FOpenColorIOConfigWrapper* ConfigWrapper = InTransform->ConfigurationOwner->GetConfigWrapper();
+		ensure(InConfigWrapper);
 
 		EOpenColorIOViewTransformDirection DisplayViewDirection;
 		if (InTransform->GetDisplayViewDirection(DisplayViewDirection))
 		{
 			return FOpenColorIOProcessorWrapper(
-				ConfigWrapper,
+				InConfigWrapper,
 				InTransform->SourceColorSpace,
 				InTransform->Display,
 				InTransform->View,
@@ -62,7 +56,7 @@ namespace
 		else
 		{
 			return FOpenColorIOProcessorWrapper(
-				ConfigWrapper,
+				InConfigWrapper,
 				InTransform->SourceColorSpace,
 				InTransform->DestinationColorSpace,
 				InTransform->GetContextKeyValues());
@@ -168,28 +162,34 @@ UOpenColorIOColorTransform::UOpenColorIOColorTransform(const FObjectInitializer&
 
 bool UOpenColorIOColorTransform::Initialize(UOpenColorIOConfiguration* InOwner, const FString& InSourceColorSpace, const FString& InDestinationColorSpace, const TMap<FString, FString>& InContextKeyValues)
 {
-	check(InOwner);
-	ConfigurationOwner = InOwner;
+	return Initialize(InSourceColorSpace, InDestinationColorSpace, InContextKeyValues);
+}
+
+bool UOpenColorIOColorTransform::Initialize(UOpenColorIOConfiguration* InOwner, const FString& InSourceColorSpace, const FString& InDisplay, const FString& InView, EOpenColorIOViewTransformDirection InDirection, const TMap<FString, FString>& InContextKeyValues)
+{
+	return Initialize(InSourceColorSpace, InDisplay, InView, InDirection, InContextKeyValues);
+}
+
+bool UOpenColorIOColorTransform::Initialize(const FString& InSourceColorSpace, const FString& InDestinationColorSpace, const TMap<FString, FString>& InContextKeyValues)
+{
 	ContextKeyValues = InContextKeyValues;
-	
+
 #if WITH_EDITOR
 	if (InSourceColorSpace == OpenColorIOWrapper::GetWorkingColorSpaceName())
 	{
 		WorkingColorSpaceTransformType = EOpenColorIOWorkingColorSpaceTransform::Source;
 	}
-	else if(InDestinationColorSpace == OpenColorIOWrapper::GetWorkingColorSpaceName())
+	else if (InDestinationColorSpace == OpenColorIOWrapper::GetWorkingColorSpaceName())
 	{
 		WorkingColorSpaceTransformType = EOpenColorIOWorkingColorSpaceTransform::Destination;
 	}
 #endif
-	
+
 	return GenerateColorTransformData(InSourceColorSpace, InDestinationColorSpace);
 }
 
-bool UOpenColorIOColorTransform::Initialize(UOpenColorIOConfiguration* InOwner, const FString& InSourceColorSpace, const FString& InDisplay, const FString& InView, EOpenColorIOViewTransformDirection InDirection, const TMap<FString, FString>& InContextKeyValues)
+bool UOpenColorIOColorTransform::Initialize(const FString& InSourceColorSpace, const FString& InDisplay, const FString& InView, EOpenColorIOViewTransformDirection InDirection, const TMap<FString, FString>& InContextKeyValues)
 {
-	check(InOwner);
-	ConfigurationOwner = InOwner;
 	ContextKeyValues = InContextKeyValues;
 
 #if WITH_EDITOR
@@ -278,9 +278,12 @@ void UOpenColorIOColorTransform::CacheResourceTextures()
 #if WITH_EDITOR
 	if (Textures.IsEmpty())
 	{
-		if (HasValidConfiguration(ConfigurationOwner))
+		const UOpenColorIOConfiguration* ConfigurationOwner = Cast<UOpenColorIOConfiguration>(GetOuter());
+		const FOpenColorIOConfigWrapper* ConfigWrapper = GetTransformConfigWrapper(ConfigurationOwner);
+
+		if (ConfigWrapper != nullptr)
 		{
-			const FOpenColorIOProcessorWrapper TransformProcessor = GetTransformProcessor(this);
+			const FOpenColorIOProcessorWrapper TransformProcessor = GetTransformProcessor(this, ConfigWrapper);
 			if (TransformProcessor.IsValid())
 			{
 				const FOpenColorIOGPUProcessorWrapper GPUProcessor = FOpenColorIOGPUProcessorWrapper(TransformProcessor);
@@ -492,9 +495,12 @@ bool UOpenColorIOColorTransform::IsTransform(const FString& InSourceColorSpace,
 #if WITH_EDITOR
 bool UOpenColorIOColorTransform::EditorTransformImage(const FImageView& InOutImage) const
 {
-	if (HasValidConfiguration(ConfigurationOwner))
+	const UOpenColorIOConfiguration* ConfigurationOwner = Cast<UOpenColorIOConfiguration>(GetOuter());
+	const FOpenColorIOConfigWrapper* ConfigWrapper = GetTransformConfigWrapper(ConfigurationOwner);
+
+	if (ConfigWrapper != nullptr)
 	{
-		const FOpenColorIOProcessorWrapper TransformProcessor = GetTransformProcessor(this);
+		const FOpenColorIOProcessorWrapper TransformProcessor = GetTransformProcessor(this, ConfigWrapper);
 		if (TransformProcessor.IsValid())
 		{
 			const FOpenColorIOCPUProcessorWrapper CPUProcessor = FOpenColorIOCPUProcessorWrapper(TransformProcessor);
@@ -512,9 +518,12 @@ bool UOpenColorIOColorTransform::EditorTransformImage(const FImageView& InOutIma
 
 bool UOpenColorIOColorTransform::EditorTransformImage(const FImageView& SrcImage, const FImageView& DestImage) const
 {
-	if (HasValidConfiguration(ConfigurationOwner))
+	const UOpenColorIOConfiguration* ConfigurationOwner = Cast<UOpenColorIOConfiguration>(GetOuter());
+	const FOpenColorIOConfigWrapper* ConfigWrapper = GetTransformConfigWrapper(ConfigurationOwner);
+
+	if (ConfigWrapper != nullptr)
 	{
-		const FOpenColorIOProcessorWrapper TransformProcessor = GetTransformProcessor(this);
+		const FOpenColorIOProcessorWrapper TransformProcessor = GetTransformProcessor(this, ConfigWrapper);
 		if (TransformProcessor.IsValid())
 		{
 			const FOpenColorIOCPUProcessorWrapper CPUProcessor = FOpenColorIOCPUProcessorWrapper(TransformProcessor);
@@ -618,9 +627,12 @@ FString UOpenColorIOColorTransform::GetTransformFriendlyName() const
 bool UOpenColorIOColorTransform::UpdateShaderInfo(FString& OutShaderCodeHash, FString& OutShaderCode, FString& OutRawConfigHash)
 {
 #if WITH_EDITOR
-	if (HasValidConfiguration(ConfigurationOwner))
+	const UOpenColorIOConfiguration* ConfigurationOwner = Cast<UOpenColorIOConfiguration>(GetOuter());
+	const FOpenColorIOConfigWrapper* ConfigWrapper = GetTransformConfigWrapper(ConfigurationOwner);
+
+	if (ConfigWrapper != nullptr)
 	{
-		const FOpenColorIOProcessorWrapper Processor = GetTransformProcessor(this);
+		const FOpenColorIOProcessorWrapper Processor = GetTransformProcessor(this, ConfigWrapper);
 		if (Processor.IsValid())
 		{
 			const FOpenColorIOGPUProcessorWrapper GPUProcessor = FOpenColorIOGPUProcessorWrapper(Processor);
@@ -800,18 +812,33 @@ void UOpenColorIOColorTransform::PostLoad()
 		}
 	}
 
-	if (!ConfigurationOwner && GetOuter())
-	{
-		UE_LOG(LogOpenColorIO, Verbose, TEXT("ConfigurationOwner is null. Assigning Outer to ConfigurationOwner."));
-		ConfigurationOwner = Cast<UOpenColorIOConfiguration>(GetOuter());
-	}
+	UOpenColorIOConfiguration* ConfigurationOwner = Cast<UOpenColorIOConfiguration>(GetOuter());
 
 	//To be able to fetch OCIO data, make sure our config owner has been postloaded.
 	if (ConfigurationOwner)
 	{
 		ConfigurationOwner->ConditionalPostLoad();
-		CacheResourceTextures();
-		CacheResourceShadersForRendering(false);
+
+		bool bIsTransformActive = false;
+
+		if (bIsDisplayViewType)
+		{
+			bIsTransformActive = ConfigurationOwner->HasTransform(SourceColorSpace, Display, View, DisplayViewDirection);
+		}
+		else
+		{
+			bIsTransformActive = ConfigurationOwner->HasTransform(SourceColorSpace, DestinationColorSpace);
+		}
+		
+		if (bIsTransformActive)
+		{
+			CacheResourceTextures();
+			CacheResourceShadersForRendering(false);
+		}
+		else
+		{
+			UE_LOG(LogOpenColorIO, Verbose, TEXT("ConfigurationOwner does not contain [%s], aborting transform load."), *GetTransformFriendlyName());
+		}
 	}
 	else
 	{
diff --git a/Engine/Plugins/Compositing/OpenColorIO/Source/OpenColorIO/Private/OpenColorIOConfiguration.cpp b/Engine/Plugins/Compositing/OpenColorIO/Source/OpenColorIO/Private/OpenColorIOConfiguration.cpp
--- a/Engine/Plugins/Compositing/OpenColorIO/Source/OpenColorIO/Private/OpenColorIOConfiguration.cpp
+++ b/Engine/Plugins/Compositing/OpenColorIO/Source/OpenColorIO/Private/OpenColorIOConfiguration.cpp
@@ -162,7 +162,7 @@ void UOpenColorIOConfiguration::ReloadExistingColorspaces()
 	TArray<FOpenColorIODisplayView> DisplayViewsToBeReloaded = DesiredDisplayViews;
 	DesiredColorSpaces.Reset();
 	DesiredDisplayViews.Reset();
-	CleanupTransforms();
+	ColorTransforms.Reset();
 	LoadConfiguration();
 
 	if (Config == nullptr)
@@ -189,14 +189,15 @@ void UOpenColorIOConfiguration::ReloadExistingColorspaces()
 
 	for (const FOpenColorIODisplayView& ExistingDisplayView : DisplayViewsToBeReloaded)
 	{
-		const FString TransformName = Config->GetDisplayViewTransformName(*ExistingDisplayView.Display, *ExistingDisplayView.View);
-		if (TransformName.IsEmpty())
+		for (int32 ViewIndex = 0; ViewIndex < Config->GetNumViews(*ExistingDisplayView.Display); ++ViewIndex)
 		{
-			// Name not found, therefore we don't need to re-add this display-view.
-			continue;
+			const FString ViewName = Config->GetViewName(*ExistingDisplayView.Display, ViewIndex);
+			
+			if (ViewName == ExistingDisplayView.View)
+			{
+				DesiredDisplayViews.Add(ExistingDisplayView);
+			}
 		}
-
-		DesiredDisplayViews.Add(ExistingDisplayView);
 	}
 
 	const UOpenColorIOSettings* Settings = GetDefault<UOpenColorIOSettings>();
@@ -354,7 +355,7 @@ void UOpenColorIOConfiguration::CreateColorTransform(const FString& InSourceColo
 	}
 
 	UOpenColorIOColorTransform* NewTransform = NewObject<UOpenColorIOColorTransform>(this, NAME_None, RF_NoFlags);
-	const bool bSuccess = NewTransform->Initialize(this, InSourceColorSpace, InDestinationColorSpace);
+	const bool bSuccess = NewTransform->Initialize(InSourceColorSpace, InDestinationColorSpace);
 
 	if (bSuccess)
 	{
@@ -380,7 +381,7 @@ void UOpenColorIOConfiguration::CreateColorTransform(const FString& InSourceColo
 	}
 
 	UOpenColorIOColorTransform* NewTransform = NewObject<UOpenColorIOColorTransform>(this, NAME_None, RF_NoFlags);
-	const bool bSuccess = NewTransform->Initialize(this, InSourceColorSpace, InDisplay, InView, InDirection);
+	const bool bSuccess = NewTransform->Initialize(InSourceColorSpace, InDisplay, InView, InDirection);
 
 	if (bSuccess)
 	{
@@ -522,7 +523,8 @@ void UOpenColorIOConfiguration::PostEditChangeProperty(FPropertyChangedEvent& Pr
 
 	if (PropertyChangedEvent.MemberProperty && PropertyChangedEvent.MemberProperty->GetFName() == GET_MEMBER_NAME_CHECKED(UOpenColorIOConfiguration, ConfigurationFile))
 	{
-		LoadConfiguration();
+		// Note: Reload calls LoadConfiguration() internally.
+		ReloadExistingColorspaces();
 	}
 	else if (PropertyChangedEvent.GetPropertyName() == GET_MEMBER_NAME_CHECKED(UOpenColorIOConfiguration, DesiredColorSpaces))
 	{
diff --git a/Engine/Plugins/Compositing/OpenColorIO/Source/OpenColorIO/Private/OpenColorIOShared.cpp b/Engine/Plugins/Compositing/OpenColorIO/Source/OpenColorIO/Private/OpenColorIOShared.cpp
--- a/Engine/Plugins/Compositing/OpenColorIO/Source/OpenColorIO/Private/OpenColorIOShared.cpp
+++ b/Engine/Plugins/Compositing/OpenColorIO/Source/OpenColorIO/Private/OpenColorIOShared.cpp
@@ -9,6 +9,7 @@
 #include "Misc/App.h"
 #include "Modules/ModuleManager.h"
 #include "OpenColorIOConfiguration.h"
+#include "OpenColorIOModule.h"
 #include "OpenColorIOShaderType.h"
 #include "OpenColorIOShader.h"
 #include "OpenColorIOShaderCompilationManager.h"
@@ -286,11 +287,11 @@ bool FOpenColorIOTransformResource::CacheShaders(const FOpenColorIOShaderMapId&
 				FOpenColorIOShaderMap::LoadFromDerivedDataCache(this, InShaderMapId, InPlatform, GameThreadShaderMap);
 				if (GameThreadShaderMap && GameThreadShaderMap->IsValid())
 				{
-					UE_LOG(LogTemp, Display, TEXT("Loaded shader %s for OCIO ColorSpace %s from DDC"), *GameThreadShaderMap->GetFriendlyName(), *GetFriendlyName());
+					UE_LOG(LogOpenColorIO, Display, TEXT("Loaded shader %s for OCIO ColorSpace %s from DDC"), *GameThreadShaderMap->GetFriendlyName(), *GetFriendlyName());
 				}
 				else
 				{
-					UE_LOG(LogTemp, Display, TEXT("Loading shader for OCIO ColorSpace %s from DDC failed. Shader needs recompile."), *GetFriendlyName());
+					UE_LOG(LogOpenColorIO, Display, TEXT("Loading shader for OCIO ColorSpace %s from DDC failed. Shader needs recompile."), *GetFriendlyName());
 				}
 #endif // WITH_EDITOR
 			}
diff --git a/Engine/Plugins/Compositing/OpenColorIO/Source/OpenColorIO/Public/OpenColorIOColorTransform.h b/Engine/Plugins/Compositing/OpenColorIO/Source/OpenColorIO/Public/OpenColorIOColorTransform.h
--- a/Engine/Plugins/Compositing/OpenColorIO/Source/OpenColorIO/Public/OpenColorIOColorTransform.h
+++ b/Engine/Plugins/Compositing/OpenColorIO/Source/OpenColorIO/Public/OpenColorIOColorTransform.h
@@ -33,9 +33,16 @@ public:
 	virtual ~UOpenColorIOColorTransform() {};
 
 public:
+	UE_DEPRECATED(5.3, "This method is deprecated, please use Initialize without the owner argument.")
 	bool Initialize(UOpenColorIOConfiguration* InOwner, const FString& InSourceColorSpace, const FString& InDestinationColorSpace, const TMap<FString, FString>& InContextKeyValues = {});
+	UE_DEPRECATED(5.3, "This method is deprecated, please use Initialize without the owner argument.")
 	bool Initialize(UOpenColorIOConfiguration* InOwner, const FString& InSourceColorSpace, const FString& InDisplay, const FString& InView, EOpenColorIOViewTransformDirection InDirection, const TMap<FString, FString>& InContextKeyValues = {});
 
+	/** Initialize resources for color space transform. */
+	bool Initialize(const FString& InSourceColorSpace, const FString& InDestinationColorSpace, const TMap<FString, FString>& InContextKeyValues = {});
+	/** Initialize resources for display-view transform. */
+	bool Initialize(const FString& InSourceColorSpace, const FString& InDisplay, const FString& InView, EOpenColorIOViewTransformDirection InDirection, const TMap<FString, FString>& InContextKeyValues = {});
+
 	/**
 	 * Cache resource shaders for cooking on the given shader platform.
 	 * If a matching shader map is not found in memory or the DDC, a new one will be compiled.
@@ -181,8 +188,11 @@ public:
 
 public:
 
-	UPROPERTY()
-	TObjectPtr<UOpenColorIOConfiguration> ConfigurationOwner;
+#if WITH_EDITORONLY_DATA
+	UE_DEPRECATED(5.3, "ConfigurationOwner is deprecated, use GetOuter() instead.")
+	UPROPERTY(Transient, meta = (DeprecatedProperty))
+	TObjectPtr<UOpenColorIOConfiguration> ConfigurationOwner_DEPRECATED;
+#endif
 
 	UPROPERTY(VisibleAnywhere, Category = "ColorSpace")
 	bool bIsDisplayViewType = false;
diff --git a/Engine/Plugins/Compositing/OpenColorIO/Source/OpenColorIOWrapper/Private/OpenColorIOWrapper.cpp b/Engine/Plugins/Compositing/OpenColorIO/Source/OpenColorIOWrapper/Private/OpenColorIOWrapper.cpp
--- a/Engine/Plugins/Compositing/OpenColorIO/Source/OpenColorIOWrapper/Private/OpenColorIOWrapper.cpp
+++ b/Engine/Plugins/Compositing/OpenColorIO/Source/OpenColorIOWrapper/Private/OpenColorIOWrapper.cpp
@@ -348,6 +348,64 @@ FString FOpenColorIOConfigWrapper::GetCacheID() const
 	return {};
 }
 
+FString FOpenColorIOConfigWrapper::GetDebugString() const
+{
+	TStringBuilder<1024> DebugStringBuilder;
+	
+#if WITH_OCIO
+	if (IsValid())
+	{
+		const OCIO_NAMESPACE::ConstConfigRcPtr& Config = Pimpl->Config;
+		if (Config->getNumColorSpaces() > 0)
+		{
+			DebugStringBuilder.Append(TEXT("** ColorSpaces **\n"));
+			
+			const int32 NumCS = Config->getNumColorSpaces(
+				OCIO_NAMESPACE::SEARCH_REFERENCE_SPACE_ALL,   // Iterate over scene & display color spaces.
+				OCIO_NAMESPACE::COLORSPACE_ALL);              // Iterate over active & inactive color spaces.
+
+			for (int32 IndexCS = 0; IndexCS < NumCS; ++IndexCS)
+			{
+				OCIO_NAMESPACE::ConstColorSpaceRcPtr cs = Config->getColorSpace(Config->getColorSpaceNameByIndex(
+					OCIO_NAMESPACE::SEARCH_REFERENCE_SPACE_ALL,
+					OCIO_NAMESPACE::COLORSPACE_ALL,
+					IndexCS));
+
+				DebugStringBuilder.Append(cs->getName());
+				DebugStringBuilder.Append(TEXT("\n"));
+			}
+
+			DebugStringBuilder.Append(TEXT("** (Display, View) pairs **\n"));
+
+			for (int32 IndexDisplay = 0; IndexDisplay < Config->getNumDisplaysAll(); ++IndexDisplay)
+			{
+				const ANSICHAR* DisplayName = Config->getDisplayAll(IndexDisplay);
+
+				// Iterate over shared views.
+				int32 NumViews = Config->getNumViews(DisplayName);
+				for (int IndexView = 0; IndexView < NumViews; ++IndexView)
+				{
+					const ANSICHAR* ViewName = Config->getView(
+						DisplayName,
+						IndexView);
+
+					DebugStringBuilder.Append(TEXT("("));
+					DebugStringBuilder.Append(DisplayName);
+					DebugStringBuilder.Append(TEXT(", "));
+					DebugStringBuilder.Append(ViewName);
+					DebugStringBuilder.Append(TEXT(")\n"));
+				}
+			}
+		}
+	}
+
+#endif // WITH_OCIO
+
+	return DebugStringBuilder.ToString();
+}
+
+
+
 FOpenColorIOProcessorWrapper::FOpenColorIOProcessorWrapper(
 	const FOpenColorIOConfigWrapper* InConfig,
 	FStringView InSourceColorSpace,
diff --git a/Engine/Plugins/Compositing/OpenColorIO/Source/OpenColorIOWrapper/Public/OpenColorIOWrapper.h b/Engine/Plugins/Compositing/OpenColorIO/Source/OpenColorIOWrapper/Public/OpenColorIOWrapper.h
--- a/Engine/Plugins/Compositing/OpenColorIO/Source/OpenColorIOWrapper/Public/OpenColorIOWrapper.h
+++ b/Engine/Plugins/Compositing/OpenColorIO/Source/OpenColorIOWrapper/Public/OpenColorIOWrapper.h
@@ -92,6 +92,9 @@ public:
 	/** Get the string hash of the config. */
 	FString GetCacheID() const;
 
+	/** Returns a config debug string. */
+	FString GetDebugString() const;
+
 private:
 
 	/** Convenience to create a config between the working color space and the default interchange one. */
