OCIO: FPlatformString::ConvertedLength for ANSI string conversion.

diff --git a/Engine/Plugins/Compositing/OpenColorIO/Source/OpenColorIOWrapper/Private/OpenColorIOWrapper.cpp b/Engine/Plugins/Compositing/OpenColorIO/Source/OpenColorIOWrapper/Private/OpenColorIOWrapper.cpp
--- a/Engine/Plugins/Compositing/OpenColorIO/Source/OpenColorIOWrapper/Private/OpenColorIOWrapper.cpp
+++ b/Engine/Plugins/Compositing/OpenColorIO/Source/OpenColorIOWrapper/Private/OpenColorIOWrapper.cpp
@@ -27,7 +27,7 @@ namespace OpenColorIOWrapper
 	// Build routine since there is no FAnsiString
 	TUniquePtr<ANSICHAR[]> MakeAnsiString(const TCHAR* Str)
 	{
-		int32 Num = FCString::Strlen(Str);
+		int32 Num = FPlatformString::ConvertedLength<ANSICHAR>(Str);
 		TUniquePtr<ANSICHAR[]> Ret =  MakeUnique<ANSICHAR[]>(Num);
 		FMemory::Memcpy(Ret.Get(), StringCast<ANSICHAR>(Str).Get(), Num);
 		return Ret;
