OCIO: Added target compression parameter.

#rnx

diff --git a/Engine/Plugins/Compositing/OpenColorIO/Source/OpenColorIOEditor/Private/OpenColorIOEditorBlueprintLibrary.cpp b/Engine/Plugins/Compositing/OpenColorIO/Source/OpenColorIOEditor/Private/OpenColorIOEditorBlueprintLibrary.cpp
--- a/Engine/Plugins/Compositing/OpenColorIO/Source/OpenColorIOEditor/Private/OpenColorIOEditorBlueprintLibrary.cpp
+++ b/Engine/Plugins/Compositing/OpenColorIO/Source/OpenColorIOEditor/Private/OpenColorIOEditorBlueprintLibrary.cpp
@@ -21,7 +21,17 @@ void UOpenColorIOEditorBlueprintLibrary::SetActiveViewportConfiguration(const FO
 }
 
 bool UOpenColorIOEditorBlueprintLibrary::ApplyColorSpaceTransformToTexture(const FOpenColorIOColorConversionSettings& ConversionSettings, UTexture* InOutTexture, bool bSynchronous)
-{	
+{
+	if (IsValid(InOutTexture))
+	{
+		return ApplyColorSpaceCompressionTransformToTexture(ConversionSettings, InOutTexture->CompressionSettings, InOutTexture, bSynchronous);
+	}
+
+	return false;
+}
+
+bool UOpenColorIOEditorBlueprintLibrary::ApplyColorSpaceCompressionTransformToTexture(const FOpenColorIOColorConversionSettings& ConversionSettings, TextureCompressionSettings TargetCompression, UTexture* InOutTexture, bool bSynchronous)
+{
 	if (IsValid(InOutTexture) && IsValid(ConversionSettings.ConfigurationSource))
 	{
 		FImage ImageMip0;
@@ -30,17 +40,18 @@ bool UOpenColorIOEditorBlueprintLibrary::ApplyColorSpaceTransformToTexture(const
 			bool bTransformSucceeded = ConversionSettings.ConfigurationSource->EditorTransformImage(ConversionSettings, ImageMip0);
 			if (bTransformSucceeded)
 			{
+				InOutTexture->PreEditChange(nullptr);
 				void* TargetData = InOutTexture->Source.LockMip(0);
 				FMemory::Memcpy(TargetData, ImageMip0.RawData.GetData(), ImageMip0.GetImageSizeBytes());
 				InOutTexture->Source.UnlockMip(0);
-				InOutTexture->Modify();
-				InOutTexture->UpdateResource();
+				InOutTexture->CompressionSettings = TargetCompression;
+				InOutTexture->PostEditChange();
 
 				if (bSynchronous)
 				{
 					FTextureCompilingManager::Get().FinishCompilation({ InOutTexture });
 				}
-				
+
 				return true;
 			}
 		}
diff --git a/Engine/Plugins/Compositing/OpenColorIO/Source/OpenColorIOEditor/Public/OpenColorIOEditorBlueprintLibrary.h b/Engine/Plugins/Compositing/OpenColorIO/Source/OpenColorIOEditor/Public/OpenColorIOEditorBlueprintLibrary.h
--- a/Engine/Plugins/Compositing/OpenColorIO/Source/OpenColorIOEditor/Public/OpenColorIOEditorBlueprintLibrary.h
+++ b/Engine/Plugins/Compositing/OpenColorIO/Source/OpenColorIOEditor/Public/OpenColorIOEditorBlueprintLibrary.h
@@ -3,6 +3,7 @@
 #pragma once
 
 #include "Kismet/BlueprintFunctionLibrary.h"
+#include "Engine/TextureDefines.h"
 
 #include "OpenColorIOEditorBlueprintLibrary.generated.h"
 
@@ -22,7 +23,6 @@ public:
 	UFUNCTION(BlueprintCallable, Category = "OpenColorIO")
 	static OPENCOLORIOEDITOR_API void SetActiveViewportConfiguration(const FOpenColorIODisplayConfiguration& InConfiguration);
 
-
 	/**
 	 * Apply a color space transform to a texture asset.
 	 *
@@ -33,4 +33,16 @@ public:
 	 */
 	UFUNCTION(BlueprintCallable, Category = "OpenColorIO")
 	static OPENCOLORIOEDITOR_API bool ApplyColorSpaceTransformToTexture(const FOpenColorIOColorConversionSettings& ConversionSettings, UTexture* InOutTexture, bool bSynchronous = false);
+
+	/**
+	 * Apply a color space transform with a target compression setting to a texture asset.
+	 *
+	 * @param ConversionSettings Color transformation settings.
+	 * @param TargetCompression Target texture compression setting.
+	 * @param InOutTexture Texture object to transform.
+	 * @param bSynchronous Whether the texture transform is a nlocking operation.
+	 * @return true upon success.
+	 */
+	UFUNCTION(BlueprintCallable, Category = "OpenColorIO")
+	static OPENCOLORIOEDITOR_API bool ApplyColorSpaceCompressionTransformToTexture(const FOpenColorIOColorConversionSettings& ConversionSettings, TextureCompressionSettings TargetCompression, UTexture* InOutTexture, bool bSynchronous = false);
 };
