OCIO: Added support for conversions to working color space from texture source settings.

diff --git a/Engine/Plugins/Compositing/OpenColorIO/Source/OpenColorIOWrapper/Private/OpenColorIOWrapper.cpp b/Engine/Plugins/Compositing/OpenColorIO/Source/OpenColorIOWrapper/Private/OpenColorIOWrapper.cpp
--- a/Engine/Plugins/Compositing/OpenColorIO/Source/OpenColorIOWrapper/Private/OpenColorIOWrapper.cpp
+++ b/Engine/Plugins/Compositing/OpenColorIO/Source/OpenColorIOWrapper/Private/OpenColorIOWrapper.cpp
@@ -5,7 +5,7 @@
 
 #include "Containers/Map.h"
 #include "ColorSpace.h"
-#include "Engine/TextureDefines.h"
+#include "Engine/Texture.h"
 #include "ImageCore.h"
 
 #if WITH_OCIO
@@ -119,42 +119,6 @@ namespace {
 }
 #endif // WITH_OCIO
 
-TUniquePtr<FOpenColorIOConfigWrapper> FOpenColorIOConfigWrapper::CreateWorkingColorSpaceToInterchangeConfig()
-{
-#if WITH_OCIO
-	TRACE_CPUPROFILER_EVENT_SCOPE(FOpenColorIOConfigWrapper::CreateWorkingColorSpaceToInterchangeConfig)
-	using namespace OCIO_NAMESPACE;
-	using namespace UE::Color;
-
-	TUniquePtr<FOpenColorIOConfigWrapper> Result = MakeUnique<FOpenColorIOConfigWrapper>();
-
-	ColorSpaceRcPtr AP0 = ColorSpace::Create();
-	AP0->setName("ACES2065-1");
-	AP0->setBitDepth(BIT_DEPTH_F32);
-	AP0->setEncoding("scene-linear");
-
-	ColorSpaceRcPtr WCS = ColorSpace::Create();
-	WCS->setName(StringCast<ANSICHAR>(OpenColorIOWrapper::GetWorkingColorSpaceName()).Get());
-	WCS->setBitDepth(BIT_DEPTH_F32);
-	WCS->setEncoding("scene-linear");
-
-	const FMatrix44d TransformMat = Transpose<double>(FColorSpaceTransform(FColorSpace::GetWorking(), FColorSpace(EColorSpace::ACESAP0)));
-	MatrixTransformRcPtr MatrixTransform = MatrixTransform::Create();
-	MatrixTransform->setMatrix(&TransformMat.M[0][0]);
-	WCS->setTransform(MatrixTransform, COLORSPACE_DIR_TO_REFERENCE);
-
-	ConfigRcPtr InterchangeConfig = Config::Create();
-	InterchangeConfig->addColorSpace(AP0);
-	InterchangeConfig->addColorSpace(WCS);
-	InterchangeConfig->setRole("aces_interchange", "ACES2065-1");
-
-	Result->Pimpl->Config = InterchangeConfig;
-
-	return Result;
-#else
-	return nullptr;
-#endif
-}
 
 FOpenColorIOConfigWrapper::FOpenColorIOConfigWrapper()
 	: Pimpl(MakePimpl<FOpenColorIOConfigPimpl, EPimplPtrMode::DeepCopy>())
@@ -404,7 +368,374 @@ FString FOpenColorIOConfigWrapper::GetDebugString() const
 	return DebugStringBuilder.ToString();
 }
 
+FOpenColorIOEngineBuiltInConfigWrapper::FOpenColorIOEngineBuiltInConfigWrapper()
+	: FOpenColorIOConfigWrapper()
+{
+#if WITH_OCIO
+	TRACE_CPUPROFILER_EVENT_SCOPE(FOpenColorIOConfigWrapper::CreateWorkingColorSpaceToInterchangeConfig)
+		using namespace OCIO_NAMESPACE;
+	using namespace UE::Color;
+
+	ConfigRcPtr StudioConfig = Config::CreateFromBuiltinConfig("studio-config-v1.0.0_aces-v1.3_ocio-v2.1")->createEditableCopy();
+
+	ColorSpaceRcPtr WCS = ColorSpace::Create();
+	WCS->setName(StringCast<ANSICHAR>(OpenColorIOWrapper::GetWorkingColorSpaceName()).Get());
+	WCS->setBitDepth(BIT_DEPTH_F32);
+	WCS->setEncoding("scene-linear");
+	// We know the scene-referred reference space is ACES2065-1, and hence the correct matrix transform.
+	const FMatrix44d TransformMat = Transpose<double>(FColorSpaceTransform(FColorSpace::GetWorking(), FColorSpace(EColorSpace::ACESAP0)));
+	MatrixTransformRcPtr MatrixTransform = MatrixTransform::Create();
+	MatrixTransform->setMatrix(&TransformMat.M[0][0]);
+	WCS->setTransform(MatrixTransform, COLORSPACE_DIR_TO_REFERENCE);
+
+	StudioConfig->addColorSpace(WCS);
+
+	Pimpl->Config = StudioConfig;
+#endif // WITH_OCIO
+}
+
+FOpenColorIOProcessorWrapper FOpenColorIOEngineBuiltInConfigWrapper::GetProcessorFromTextureColorSettings(const FTextureSourceColorSettings& InTextureColorSettings)
+{
+	return GetProcessorFromTextureColorSettings(InTextureColorSettings, UE::Color::FColorSpace::GetWorking());
+}
+
+FOpenColorIOProcessorWrapper FOpenColorIOEngineBuiltInConfigWrapper::GetProcessorFromTextureColorSettings(const FTextureSourceColorSettings& InTextureColorSettings, const UE::Color::FColorSpace& InTargetColorSpace)
+{
+	const uint32 SettingsHash = FCrc::MemCrc32(&InTextureColorSettings, sizeof(FTextureSourceColorSettings));
+	const FString TransformName = FString::Printf(TEXT("WCS%u"), SettingsHash);
+
+#if WITH_OCIO
+	using namespace OCIO_NAMESPACE;
+	using namespace UE::Color;
+
+	if (Pimpl->Config->getNamedTransform(StringCast<ANSICHAR>(*TransformName).Get()) == nullptr)
+	{
+		GroupTransformRcPtr TransformToWCS = GroupTransform::Create();
+		NamedTransformRcPtr ParentTransform = NamedTransform::Create();
+		ParentTransform->setName(StringCast<ANSICHAR>(*TransformName).Get());
+
+		switch (InTextureColorSettings.EncodingOverride)
+		{
+		case ETextureSourceEncoding::TSE_None:
+		case ETextureSourceEncoding::TSE_Linear:
+			ParentTransform->setEncoding("scene-linear");
+			break;
+		case ETextureSourceEncoding::TSE_sRGB:
+		{
+			ExponentWithLinearTransformRcPtr ChildTransform = ExponentWithLinearTransform::Create();
+			ChildTransform->setGamma({ 2.4, 2.4, 2.4, 1.0 });
+			ChildTransform->setOffset({ 0.055, 0.055, 0.055, 0.0 });
 
+			TransformToWCS->appendTransform(ChildTransform);
+			ParentTransform->setEncoding("sdr-video");
+		}
+		break;
+		case ETextureSourceEncoding::TSE_ST2084:
+		{
+			BuiltinTransformRcPtr ChildTransform = BuiltinTransform::Create();
+			ChildTransform->setStyle("CURVE - ST-2084_to_LINEAR");
+			TransformToWCS->appendTransform(ChildTransform);
+
+			// By default ocio returns nits/100
+			MatrixTransformRcPtr RescaleTransform = MatrixTransform::Create();
+			FMatrix44d ScaleMatrix = FMatrix44d::Identity.ApplyScale(100.0);
+			RescaleTransform->setMatrix(&ScaleMatrix.M[0][0]);
+
+			TransformToWCS->appendTransform(RescaleTransform);
+			ParentTransform->setEncoding("hdr-video");
+		}
+		break;
+		case ETextureSourceEncoding::TSE_Gamma22:
+		{
+			ExponentTransformRcPtr ChildTransform = ExponentTransform::Create();
+			ChildTransform->setValue({ 2.2, 2.2, 2.2, 1.0 });
+			ChildTransform->setNegativeStyle(NEGATIVE_PASS_THRU);
+
+			TransformToWCS->appendTransform(ChildTransform);
+			ParentTransform->setEncoding("sdr-video");
+		}
+		break;
+		case ETextureSourceEncoding::TSE_BT1886:
+		{
+			ExponentTransformRcPtr ChildTransform = ExponentTransform::Create();
+			ChildTransform->setValue({ 2.4, 2.4, 2.4, 1.0 });
+			ChildTransform->setNegativeStyle(NEGATIVE_PASS_THRU);
+
+			TransformToWCS->appendTransform(ChildTransform);
+			ParentTransform->setEncoding("sdr-video");
+		}
+		break;
+		case ETextureSourceEncoding::TSE_Gamma26:
+		{
+			ExponentTransformRcPtr ChildTransform = ExponentTransform::Create();
+			ChildTransform->setValue({ 2.6, 2.6, 2.6, 1.0 });
+			ChildTransform->setNegativeStyle(NEGATIVE_PASS_THRU);
+
+			TransformToWCS->appendTransform(ChildTransform);
+			ParentTransform->setEncoding("sdr-video");
+		}
+		break;
+		case ETextureSourceEncoding::TSE_Cineon:
+		{
+			const double BlackOffset = FGenericPlatformMath::Pow(10.0, (95.0 - 685.0) / 300.0);
+			const double LinSideSlope = 1.0 - BlackOffset;
+			const double LinSideOffset = BlackOffset;
+			static constexpr double LogSideSlope = 300.0 / 1023.0;
+			static constexpr double LogSideOffset = 685.0 / 1023.0;
+			static constexpr double Base = 10.;
+
+			LogAffineTransformRcPtr ChildTransform = LogAffineTransform::Create();
+			ChildTransform->setLinSideSlopeValue({ LinSideSlope, LinSideSlope,LinSideSlope });
+			ChildTransform->setLinSideOffsetValue({ LinSideOffset, LinSideOffset,LinSideOffset });
+			ChildTransform->setLogSideSlopeValue({ LogSideSlope, LogSideSlope, LogSideSlope });
+			ChildTransform->setLogSideOffsetValue({ LogSideOffset, LogSideOffset, LogSideOffset });
+			ChildTransform->setBase(Base);
+			ChildTransform->setDirection(TRANSFORM_DIR_INVERSE);
+
+			TransformToWCS->appendTransform(ChildTransform);
+			ParentTransform->setEncoding("log");
+		}
+		break;
+		case ETextureSourceEncoding::TSE_REDLog:
+		{
+			const double BlackOffset = FGenericPlatformMath::Pow(10.0, (0.0 - 1023.0) / 511.0);
+			const double LinSideSlope = 1.0 - BlackOffset;
+			const double LinSideOffset = BlackOffset;
+			static constexpr double LogSideSlope = 511.0 / 1023.0;
+			static constexpr double LogSideOffset = 1.0;
+			static constexpr double Base = 10.;
+
+			LogAffineTransformRcPtr ChildTransform = LogAffineTransform::Create();
+			ChildTransform->setLinSideSlopeValue({ LinSideSlope, LinSideSlope,LinSideSlope });
+			ChildTransform->setLinSideOffsetValue({ LinSideOffset, LinSideOffset,LinSideOffset });
+			ChildTransform->setLogSideSlopeValue({ LogSideSlope, LogSideSlope, LogSideSlope });
+			ChildTransform->setLogSideOffsetValue({ LogSideOffset, LogSideOffset, LogSideOffset });
+			ChildTransform->setBase(Base);
+			ChildTransform->setDirection(TRANSFORM_DIR_INVERSE);
+
+			TransformToWCS->appendTransform(ChildTransform);
+			ParentTransform->setEncoding("log");
+		}
+		break;
+		case ETextureSourceEncoding::TSE_REDLog3G10:
+		{
+			static constexpr double LinSideSlope = 155.975327;
+			static constexpr double LinSideOffset = 0.01 * LinSideSlope + 1.0;
+			static constexpr double LogSideSlope = 0.224282;
+			static constexpr double LogSideOffset = 0.0;
+			static constexpr double LinSideBreak = -0.01;
+			static constexpr double Base = 10.;
+
+			LogCameraTransformRcPtr ChildTransform = LogCameraTransform::Create({ LinSideBreak, LinSideBreak,LinSideBreak });
+			ChildTransform->setLinSideSlopeValue({ LinSideSlope, LinSideSlope,LinSideSlope });
+			ChildTransform->setLinSideOffsetValue({ LinSideOffset, LinSideOffset,LinSideOffset });
+			ChildTransform->setLogSideSlopeValue({ LogSideSlope, LogSideSlope, LogSideSlope });
+			ChildTransform->setLogSideOffsetValue({ LogSideOffset, LogSideOffset, LogSideOffset });
+			ChildTransform->setBase(Base);
+			ChildTransform->setDirection(TRANSFORM_DIR_INVERSE);
+
+			TransformToWCS->appendTransform(ChildTransform);
+			ParentTransform->setEncoding("log");
+		}
+		break;
+		case ETextureSourceEncoding::TSE_SLog1:
+		{
+			static constexpr double LinSideSlope = 1.0 / 0.9;
+			static constexpr double LinSideOffset = 0.037584;
+			static constexpr double LogSideSlope = 0.432699 * 219.0 * 4.0 / 1023.0;
+			static constexpr double LogSideOffset = ((0.616596 + 0.03) * 219.0 + 16.0) * 4.0 / 1023.0;
+			static constexpr double Base = 10.;
+
+			LogAffineTransformRcPtr ChildTransform = LogAffineTransform::Create();
+			ChildTransform->setLinSideSlopeValue({ LinSideSlope, LinSideSlope,LinSideSlope });
+			ChildTransform->setLinSideOffsetValue({ LinSideOffset, LinSideOffset,LinSideOffset });
+			ChildTransform->setLogSideSlopeValue({ LogSideSlope, LogSideSlope, LogSideSlope });
+			ChildTransform->setLogSideOffsetValue({ LogSideOffset, LogSideOffset, LogSideOffset });
+			ChildTransform->setBase(Base);
+			ChildTransform->setDirection(TRANSFORM_DIR_INVERSE);
+
+			TransformToWCS->appendTransform(ChildTransform);
+			ParentTransform->setEncoding("log");
+		}
+		break;
+		case ETextureSourceEncoding::TSE_SLog2:
+		{
+			static constexpr double LinSideSlope = 155.0 / 197.1;
+			static constexpr double LinSideOffset = 0.037584;
+			static constexpr double LogSideSlope = 876.0 * 0.432699 / 1023.0;
+			static constexpr double LogSideOffset = (64.0 + 876.0 * 0.646596) / 1023.0;
+			static constexpr double LinSideBreak = 0.0;
+			static constexpr double LinearSlope = 876.0 * (3.53881278538813f / 0.9) / 1023.0;
+			static constexpr double Base = 10.;
+
+			LogCameraTransformRcPtr ChildTransform = LogCameraTransform::Create({ LinSideBreak, LinSideBreak,LinSideBreak });
+			ChildTransform->setLinSideSlopeValue({ LinSideSlope, LinSideSlope,LinSideSlope });
+			ChildTransform->setLinSideOffsetValue({ LinSideOffset, LinSideOffset,LinSideOffset });
+			ChildTransform->setLogSideSlopeValue({ LogSideSlope, LogSideSlope, LogSideSlope });
+			ChildTransform->setLogSideOffsetValue({ LogSideOffset, LogSideOffset, LogSideOffset });
+			ChildTransform->setLinearSlopeValue({ LinearSlope, LinearSlope, LinearSlope });
+			ChildTransform->setBase(Base);
+			ChildTransform->setDirection(TRANSFORM_DIR_INVERSE);
+
+			TransformToWCS->appendTransform(ChildTransform);
+			ParentTransform->setEncoding("log");
+		}
+		break;
+		case ETextureSourceEncoding::TSE_SLog3:
+		{
+			static constexpr double LinSideSlope = 5.26315789473684;
+			static constexpr double LinSideOffset = 0.0526315789473684;
+			static constexpr double LogSideSlope = 0.255620723362659;
+			static constexpr double LogSideOffset = 0.410557184750733;
+			static constexpr double LinSideBreak = 0.01125;
+			static constexpr double LinearSlope = 6.62194371177582;
+			static constexpr double Base = 10.;
+
+			LogCameraTransformRcPtr ChildTransform = LogCameraTransform::Create({ LinSideBreak, LinSideBreak,LinSideBreak });
+			ChildTransform->setLinSideSlopeValue({ LinSideSlope, LinSideSlope,LinSideSlope });
+			ChildTransform->setLinSideOffsetValue({ LinSideOffset, LinSideOffset,LinSideOffset });
+			ChildTransform->setLogSideSlopeValue({ LogSideSlope, LogSideSlope, LogSideSlope });
+			ChildTransform->setLogSideOffsetValue({ LogSideOffset, LogSideOffset, LogSideOffset });
+			ChildTransform->setLinearSlopeValue({ LinearSlope, LinearSlope, LinearSlope });
+			ChildTransform->setBase(Base);
+			ChildTransform->setDirection(TRANSFORM_DIR_INVERSE);
+
+			TransformToWCS->appendTransform(ChildTransform);
+			ParentTransform->setEncoding("log");
+		}
+		break;
+		case ETextureSourceEncoding::TSE_AlexaV3LogC:
+		{
+			static constexpr double LinSideSlope = 5.55555555555556;
+			static constexpr double LinSideOffset = 0.0522722750251688;
+			static constexpr double LogSideSlope = 0.247189638318671;
+			static constexpr double LogSideOffset = 0.385536998692443;
+			static constexpr double LinSideBreak = 0.0105909904954696;
+			static constexpr double Base = 10.;
+
+			LogCameraTransformRcPtr ChildTransform = LogCameraTransform::Create({ LinSideBreak, LinSideBreak,LinSideBreak });
+			ChildTransform->setLinSideSlopeValue({ LinSideSlope, LinSideSlope,LinSideSlope });
+			ChildTransform->setLinSideOffsetValue({ LinSideOffset, LinSideOffset,LinSideOffset });
+			ChildTransform->setLogSideSlopeValue({ LogSideSlope, LogSideSlope, LogSideSlope });
+			ChildTransform->setLogSideOffsetValue({ LogSideOffset, LogSideOffset, LogSideOffset });
+			ChildTransform->setBase(Base);
+			ChildTransform->setDirection(TRANSFORM_DIR_INVERSE);
+
+			TransformToWCS->appendTransform(ChildTransform);
+			ParentTransform->setEncoding("log");
+		}
+		break;
+		case ETextureSourceEncoding::TSE_CanonLog:
+		{
+			static constexpr double LinSideSlope = 10.1596;
+			static constexpr double LinSideOffset = 1.0;
+			static constexpr double LogSideSlope = 0.529136;
+			static constexpr double LogSideOffset = 0.0730597;
+			static constexpr double LinSideBreak = 0.0;
+			static constexpr double Base = 10.;
+
+			LogAffineTransformRcPtr ChildTransform = LogAffineTransform::Create();
+			ChildTransform->setLinSideSlopeValue({ LinSideSlope, LinSideSlope,LinSideSlope });
+			ChildTransform->setLinSideOffsetValue({ LinSideOffset, LinSideOffset,LinSideOffset });
+			ChildTransform->setLogSideSlopeValue({ LogSideSlope, LogSideSlope, LogSideSlope });
+			ChildTransform->setLogSideOffsetValue({ LogSideOffset, LogSideOffset, LogSideOffset });
+			ChildTransform->setBase(Base);
+			ChildTransform->setDirection(TRANSFORM_DIR_INVERSE);
+
+			TransformToWCS->appendTransform(ChildTransform);
+			ParentTransform->setEncoding("log");
+		}
+		break;
+		case ETextureSourceEncoding::TSE_ProTune:
+		{
+			static constexpr double LinSideSlope = 112.0;
+			static constexpr double LinSideOffset = 1.0f;
+			const double LogSideSlope = 1.0 / FGenericPlatformMath::Loge(113.0);
+
+			LogAffineTransformRcPtr ChildTransform = LogAffineTransform::Create();
+			ChildTransform->setLinSideSlopeValue({ LinSideSlope, LinSideSlope,LinSideSlope });
+			ChildTransform->setLinSideOffsetValue({ LinSideOffset, LinSideOffset,LinSideOffset });
+			ChildTransform->setLogSideSlopeValue({ LogSideSlope, LogSideSlope, LogSideSlope });
+			ChildTransform->setBase(UE_DOUBLE_EULERS_NUMBER);
+			ChildTransform->setDirection(TRANSFORM_DIR_INVERSE);
+
+			TransformToWCS->appendTransform(ChildTransform);
+			ParentTransform->setEncoding("log");
+		}
+		break;
+		case ETextureSourceEncoding::TSE_VLog:
+		{
+			static constexpr double LinSideSlope = 1.0;
+			static constexpr double LinSideOffset = 0.00873;
+			static constexpr double LogSideSlope = 0.241514;
+			static constexpr double LogSideOffset = 0.598206;
+			static constexpr double LinSideBreak = 0.01;
+			static constexpr double Base = 10.;
+			// Note: this is not in the studio config
+			// static constexpr double LinearSlope = 5.6;
+
+			LogCameraTransformRcPtr ChildTransform = LogCameraTransform::Create({ LinSideBreak, LinSideBreak,LinSideBreak });
+			ChildTransform->setLinSideSlopeValue({ LinSideSlope, LinSideSlope,LinSideSlope });
+			ChildTransform->setLinSideOffsetValue({ LinSideOffset, LinSideOffset,LinSideOffset });
+			ChildTransform->setLogSideSlopeValue({ LogSideSlope, LogSideSlope, LogSideSlope });
+			ChildTransform->setLogSideOffsetValue({ LogSideOffset, LogSideOffset, LogSideOffset });
+			//ChildTransform->setLinearSlopeValue({ LinearSlope, LinearSlope, LinearSlope });
+			ChildTransform->setBase(Base);
+			ChildTransform->setDirection(TRANSFORM_DIR_INVERSE);
+
+			TransformToWCS->appendTransform(ChildTransform);
+			ParentTransform->setEncoding("log");
+		}
+		break;
+
+		default:
+			checkNoEntry();
+			break;
+		}
+
+		const EChromaticAdaptationMethod ChromaticAdapation = static_cast<EChromaticAdaptationMethod>(InTextureColorSettings.ChromaticAdaptationMethod);
+
+		switch (InTextureColorSettings.ColorSpace)
+		{
+		case ETextureColorSpace::TCS_None:
+			break;
+		case ETextureColorSpace::TCS_Custom:
+		{
+			const FColorSpace SourceColorSpace = FColorSpace
+			(
+				InTextureColorSettings.RedChromaticityCoordinate,
+				InTextureColorSettings.GreenChromaticityCoordinate,
+				InTextureColorSettings.BlueChromaticityCoordinate,
+				InTextureColorSettings.WhiteChromaticityCoordinate
+			);
+			MatrixTransformRcPtr MatrixTransform = MatrixTransform::Create();
+			const FMatrix44d ToWorkingMat = Transpose<double>(FColorSpaceTransform(SourceColorSpace, InTargetColorSpace, ChromaticAdapation));
+			MatrixTransform->setMatrix(&ToWorkingMat.M[0][0]);
+			TransformToWCS->appendTransform(MatrixTransform);
+		}
+		break;
+		default:
+		{
+			const FColorSpace SourceColorSpace = FColorSpace(static_cast<EColorSpace>(InTextureColorSettings.ColorSpace));
+			const FMatrix44d ToWorkingMat = Transpose<double>(FColorSpaceTransform(SourceColorSpace, InTargetColorSpace, ChromaticAdapation));
+			MatrixTransformRcPtr MatrixTransform = MatrixTransform::Create();
+			MatrixTransform->setMatrix(&ToWorkingMat.M[0][0]);
+			TransformToWCS->appendTransform(MatrixTransform);
+		}
+		break;
+		}
+
+		ParentTransform->setTransform(TransformToWCS, TRANSFORM_DIR_FORWARD);
+
+		// Update builtin config
+		ConfigRcPtr NewConfig = Pimpl->Config->createEditableCopy();
+		NewConfig->addNamedTransform(ParentTransform);
+		Pimpl->Config = NewConfig;
+	}
+#endif // WITH_OCIO
+
+	return FOpenColorIOProcessorWrapper(this, TransformName);
+}
 
 FOpenColorIOProcessorWrapper::FOpenColorIOProcessorWrapper(
 	const FOpenColorIOConfigWrapper* InConfig,
@@ -432,7 +763,7 @@ FOpenColorIOProcessorWrapper::FOpenColorIOProcessorWrapper(
 		if (OwnerConfig != nullptr && OwnerConfig->IsValid())
 		{
 			const OCIO_NAMESPACE::ConstConfigRcPtr& Config = OwnerConfig->Pimpl->Config;
-			 OCIO_NAMESPACE::ContextRcPtr Context = Config->getCurrentContext()->createEditableCopy();
+			OCIO_NAMESPACE::ContextRcPtr Context = Config->getCurrentContext()->createEditableCopy();
 
 			for (const TPair<FString, FString>& KeyValue : InContextKeyValues)
 			{
@@ -510,6 +841,41 @@ FOpenColorIOProcessorWrapper::FOpenColorIOProcessorWrapper(
 #endif // WITH_OCIO
 }
 
+FOpenColorIOProcessorWrapper::FOpenColorIOProcessorWrapper(const FOpenColorIOConfigWrapper* InConfig, FStringView InNamedTransform, bool bInverseDirection, const TMap<FString, FString>& InContextKeyValues)
+	: Pimpl(MakePimpl<FOpenColorIOProcessorPimpl, EPimplPtrMode::DeepCopy>())
+	, OwnerConfig(InConfig)
+	, WorkingColorSpaceTransformType(EOpenColorIOWorkingColorSpaceTransform::None)
+{
+#if WITH_OCIO
+#if !PLATFORM_EXCEPTIONS_DISABLED
+	try
+#endif
+	{
+		if (OwnerConfig != nullptr && OwnerConfig->IsValid())
+		{
+			const OCIO_NAMESPACE::ConstConfigRcPtr& Config = OwnerConfig->Pimpl->Config;
+			OCIO_NAMESPACE::ContextRcPtr Context = Config->getCurrentContext()->createEditableCopy();
+
+			for (const TPair<FString, FString>& KeyValue : InContextKeyValues)
+			{
+				Context->setStringVar(TCHAR_TO_ANSI(*KeyValue.Key), TCHAR_TO_ANSI(*KeyValue.Value));
+			}
+
+			Pimpl->Processor = Config->getProcessor(
+				StringCast<ANSICHAR>(InNamedTransform.GetData()).Get(),
+				static_cast<OCIO_NAMESPACE::TransformDirection>(bInverseDirection)
+			);
+		}
+	}
+#if !PLATFORM_EXCEPTIONS_DISABLED
+	catch (OCIO_NAMESPACE::Exception& Exc)
+	{
+		UE_LOG(LogOpenColorIOWrapper, Log, TEXT("Failed to create processor for [%s]. Error message: %s"), InNamedTransform.GetData(), StringCast<TCHAR>(Exc.what()).Get());
+	}
+#endif
+#endif // WITH_OCIO
+}
+
 bool FOpenColorIOProcessorWrapper::IsValid() const
 {
 #if WITH_OCIO
@@ -529,6 +895,36 @@ bool FOpenColorIOCPUProcessorWrapper::IsValid() const
 	return ParentProcessor.IsValid();
 }
 
+bool FOpenColorIOCPUProcessorWrapper::TransformColor(FLinearColor& InOutColor) const
+{
+#if WITH_OCIO
+	TRACE_CPUPROFILER_EVENT_SCOPE(FOpenColorIOCPUProcessorWrapper::TransformImage)
+#if !PLATFORM_EXCEPTIONS_DISABLED
+	try
+#endif
+	{
+		using namespace OCIO_NAMESPACE;
+
+		if (IsValid())
+		{
+			// Apply the main color transformation
+			ConstCPUProcessorRcPtr CPUProcessor = ParentProcessor.Pimpl->Processor->getOptimizedCPUProcessor(BIT_DEPTH_F32, BIT_DEPTH_F32, OPTIMIZATION_DEFAULT);
+			CPUProcessor->applyRGBA(&InOutColor.R);
+
+			return true;
+		}
+	}
+#if !PLATFORM_EXCEPTIONS_DISABLED
+	catch (OCIO_NAMESPACE::Exception& Exc)
+	{
+		UE_LOG(LogOpenColorIOWrapper, Log, TEXT("Failed to transform color. Error message: %s"), StringCast<TCHAR>(Exc.what()).Get());
+	}
+#endif
+#endif // WITH_OCIO
+
+	return false;
+}
+
 bool FOpenColorIOCPUProcessorWrapper::TransformImage(const FImageView& InOutImage) const
 {
 #if WITH_OCIO
@@ -545,7 +941,7 @@ bool FOpenColorIOCPUProcessorWrapper::TransformImage(const FImageView& InOutImag
 			if (ImageDesc)
 			{
 				const TUniquePtr<ANSICHAR[]> AnsiWorkingColorSpaceName = OpenColorIOWrapper::MakeAnsiString(OpenColorIOWrapper::GetWorkingColorSpaceName());
-				ConstConfigRcPtr	InterchangeConfig = IOpenColorIOWrapperModule::Get().GetWorkingColorSpaceToInterchangeConfig()->Pimpl->Config;
+				ConstConfigRcPtr	InterchangeConfig = IOpenColorIOWrapperModule::Get().GetEngineBuiltInConfig()->Pimpl->Config;
 				ConstConfigRcPtr Config = ParentProcessor.OwnerConfig->Pimpl->Config;
 
 				BitDepth BitDepth = ImageDesc->getBitDepth();
@@ -612,7 +1008,7 @@ bool FOpenColorIOCPUProcessorWrapper::TransformImage(const FImageView& SrcImage,
 			if (SrcImageDesc && DestImageDesc)
 			{
 				const TUniquePtr<ANSICHAR[]> AnsiWorkingColorSpaceName = OpenColorIOWrapper::MakeAnsiString(OpenColorIOWrapper::GetWorkingColorSpaceName());
-				ConstConfigRcPtr	InterchangeConfig = IOpenColorIOWrapperModule::Get().GetWorkingColorSpaceToInterchangeConfig()->Pimpl->Config;
+				ConstConfigRcPtr	InterchangeConfig = IOpenColorIOWrapperModule::Get().GetEngineBuiltInConfig()->Pimpl->Config;
 				ConstConfigRcPtr	Config = ParentProcessor.OwnerConfig->Pimpl->Config;
 
 				BitDepth SrcBitDepth = SrcImageDesc->getBitDepth();
diff --git a/Engine/Plugins/Compositing/OpenColorIO/Source/OpenColorIOWrapper/Private/OpenColorIOWrapperModule.cpp b/Engine/Plugins/Compositing/OpenColorIO/Source/OpenColorIOWrapper/Private/OpenColorIOWrapperModule.cpp
--- a/Engine/Plugins/Compositing/OpenColorIO/Source/OpenColorIOWrapper/Private/OpenColorIOWrapperModule.cpp
+++ b/Engine/Plugins/Compositing/OpenColorIO/Source/OpenColorIOWrapper/Private/OpenColorIOWrapperModule.cpp
@@ -19,7 +19,7 @@ public:
 	{
 		FOpenColorIOLibHandler::Initialize();
 
-		InterchangeConfig = FOpenColorIOConfigWrapper::CreateWorkingColorSpaceToInterchangeConfig();
+		EngineBuiltInConfig = MakeUnique<FOpenColorIOEngineBuiltInConfigWrapper>();
 	}
 
 	virtual void ShutdownModule() override
@@ -29,27 +29,19 @@ public:
 	//~ End IModuleInterface interface
 
 	//~ Begin IOpenColorIOWrapperModule interface
-	virtual const FOpenColorIOConfigWrapper* GetWorkingColorSpaceToInterchangeConfig() const override
+	virtual FOpenColorIOEngineBuiltInConfigWrapper* GetEngineBuiltInConfig() override
 	{
-		return InterchangeConfig.Get();
+		return EngineBuiltInConfig.Get();
 	}
 
-	virtual void LoadGlobalConfig(FStringView InFilePath) override
+	virtual const FOpenColorIOEngineBuiltInConfigWrapper* GetEngineBuiltInConfig() const override
 	{
-		FOpenColorIOConfigWrapper::FInitializationOptions Opt;
-		Opt.bAddWorkingColorSpace = true;
-		EngineConfig = MakeUnique<FOpenColorIOConfigWrapper>(InFilePath, Opt);
-	}
-
-	virtual const FOpenColorIOConfigWrapper* GetGlobalConfig() const override
-	{
-		return EngineConfig.Get();
+		return EngineBuiltInConfig.Get();
 	}
 	//~ End IOpenColorIOWrapperModule interface
 private:
 	
-	TUniquePtr<FOpenColorIOConfigWrapper> InterchangeConfig = nullptr;
-	TUniquePtr<FOpenColorIOConfigWrapper> EngineConfig = nullptr;
+	TUniquePtr<FOpenColorIOEngineBuiltInConfigWrapper> EngineBuiltInConfig = nullptr;
 };
 
 IMPLEMENT_MODULE(FOpenColorIOWrapperModule, OpenColorIOWrapper);
diff --git a/Engine/Plugins/Compositing/OpenColorIO/Source/OpenColorIOWrapper/Private/Tests/OpenColorIOWrapperTests.cpp b/Engine/Plugins/Compositing/OpenColorIO/Source/OpenColorIOWrapper/Private/Tests/OpenColorIOWrapperTests.cpp
new file mode 100644
--- /dev/null
+++ b/Engine/Plugins/Compositing/OpenColorIO/Source/OpenColorIOWrapper/Private/Tests/OpenColorIOWrapperTests.cpp
@@ -0,0 +1,64 @@
+// Copyright Epic Games, Inc. All Rights Reserved.
+
+#include "CoreTypes.h"
+#include "ColorSpace.h"
+#include "Logging/LogMacros.h"
+#include "Math/NumericLimits.h"
+#include "Math/UnrealMathUtility.h"
+#include "Misc/AutomationTest.h"
+#include "OpenColorIOWrapper.h"
+#include "Engine/Texture.h"
+#include "TransferFunctions.h"
+
+#if WITH_DEV_AUTOMATION_TESTS
+
+#if WITH_OCIO
+THIRD_PARTY_INCLUDES_START
+#include "OpenColorIO/OpenColorIO.h"
+THIRD_PARTY_INCLUDES_END
+#endif
+
+DEFINE_LOG_CATEGORY_STATIC(LogUnrealOpenColorIOTest, Log, All);
+
+IMPLEMENT_SIMPLE_AUTOMATION_TEST(FOpenColorIOTransferFunctionsTest, "System.OpenColorIO.DecodeToWorkingColorSpace", EAutomationTestFlags::EditorContext | EAutomationTestFlags::EngineFilter)
+bool FOpenColorIOTransferFunctionsTest::RunTest(const FString& Parameters)
+{
+	bool bSuccessful = true;
+
+#if WITH_OCIO
+	using namespace OCIO_NAMESPACE;
+	using namespace UE::Color;
+
+	const FLinearColor TestColor = FLinearColor(0.9f, 0.5f, 0.2f, 1.0f);
+	FOpenColorIOEngineBuiltInConfigWrapper TestConfig = {};
+
+	for (uint8 TestEncoding = static_cast<uint8>(ETextureSourceEncoding::TSE_None); TestEncoding < static_cast<uint8>(ETextureSourceEncoding::TSE_MAX); ++TestEncoding)
+	{
+		FLinearColor Expected = UE::Color::Decode(static_cast<EEncoding>(TestEncoding), TestColor);
+
+		FTextureSourceColorSettings TestSettings = {};
+		TestSettings.EncodingOverride = static_cast<ETextureSourceEncoding>(TestEncoding);
+		TestSettings.ColorSpace = ETextureColorSpace::TCS_None;
+
+		FOpenColorIOProcessorWrapper Processor = TestConfig.GetProcessorFromTextureColorSettings(TestSettings);
+		FOpenColorIOCPUProcessorWrapper ProcessorCPU(Processor);
+
+		FLinearColor Actual = TestColor;
+		ProcessorCPU.TransformColor(Actual);
+
+		// Note: We make the tolerance relative to the values themselves to account for larger values in PQ.
+		const float Tolerance = UE_KINDA_SMALL_NUMBER * 0.5f * (Actual.R + Expected.R);
+
+		if (!Actual.Equals(Expected, Tolerance))
+		{
+			const FString TestNameToPrint = FString::Printf(TEXT("OpenColorIO: %u:%u"), (uint32)TestSettings.EncodingOverride, (uint32)TestSettings.ColorSpace);
+			AddError(FString::Printf(TEXT("Expected '%s' to be %s, but it was %s."), *TestNameToPrint, *Expected.ToString(), *Actual.ToString()), 1);
+			bSuccessful = false;
+		}
+	}
+#endif
+
+	return bSuccessful;
+}
+
+#endif
\ No newline at end of file
diff --git a/Engine/Plugins/Compositing/OpenColorIO/Source/OpenColorIOWrapper/Public/OpenColorIOWrapper.h b/Engine/Plugins/Compositing/OpenColorIO/Source/OpenColorIOWrapper/Public/OpenColorIOWrapper.h
--- a/Engine/Plugins/Compositing/OpenColorIO/Source/OpenColorIOWrapper/Public/OpenColorIOWrapper.h
+++ b/Engine/Plugins/Compositing/OpenColorIO/Source/OpenColorIOWrapper/Public/OpenColorIOWrapper.h
@@ -13,8 +13,17 @@
 #include "OpenColorIOWrapperDefines.h"
 
 struct FImageView;
+struct FTextureSourceColorSettings;
 enum TextureFilter : int;
 
+namespace UE {
+	namespace Color {
+		enum class EEncoding: uint8;
+		enum class EColorSpace : uint8;
+		class FColorSpace;
+	}
+}
+
 namespace OpenColorIOWrapper
 {
 	/** Color space name of the engine's working color space we insert in OpenColorIO configs. */
@@ -36,7 +45,7 @@ namespace OpenColorIOWrapper
 	OPENCOLORIOWRAPPER_API const TCHAR* GetVersion();
 }
 
-class OPENCOLORIOWRAPPER_API FOpenColorIOConfigWrapper final
+class OPENCOLORIOWRAPPER_API FOpenColorIOConfigWrapper
 {
 public:
 
@@ -57,7 +66,9 @@ public:
 	* @param InFilePath Config absolute file path.
 	* @param InOptions Initialization options.
 	*/
-	FOpenColorIOConfigWrapper(FStringView InFilePath, FInitializationOptions InOptions);
+	explicit FOpenColorIOConfigWrapper(FStringView InFilePath, FInitializationOptions InOptions);
+
+	virtual ~FOpenColorIOConfigWrapper() = default;
 
 	/** Valid when the native config has been successfully created and isn't null. */
 	bool IsValid() const;
@@ -95,19 +106,39 @@ public:
 	/** Returns a config debug string. */
 	FString GetDebugString() const;
 
-private:
-
-	/** Convenience to create a config between the working color space and the default interchange one. */
-	static TUniquePtr<FOpenColorIOConfigWrapper> CreateWorkingColorSpaceToInterchangeConfig();
+protected:
 
 	TPimplPtr<struct FOpenColorIOConfigPimpl, EPimplPtrMode::DeepCopy> Pimpl;
 
-	friend class FOpenColorIOWrapperModule;
 	friend class FOpenColorIOProcessorWrapper;
 	friend class FOpenColorIOCPUProcessorWrapper;
 	friend class FOpenColorIOGPUProcessorWrapper;
 };
 
+class OPENCOLORIOWRAPPER_API FOpenColorIOEngineBuiltInConfigWrapper : public FOpenColorIOConfigWrapper
+{
+public:
+	/**
+	* Constructor.
+	*/
+	FOpenColorIOEngineBuiltInConfigWrapper();
+	FOpenColorIOEngineBuiltInConfigWrapper(FStringView InFilePath, FInitializationOptions InOptions) = delete;
+
+	/**
+	 * Get the processor to working color space from source texture settings.
+	 * @param InTextureColorSettings Source encoding and colorspace.
+	*/
+	class FOpenColorIOProcessorWrapper GetProcessorFromTextureColorSettings(const FTextureSourceColorSettings& InTextureColorSettings);
+
+	/**
+	 * Get the processor to target color space from source texture settings.
+	 * @param InTextureColorSettings Source encoding and colorspace.
+	 * @param InTargetColorSpace Target scene-linear color space.
+	*/
+	class FOpenColorIOProcessorWrapper GetProcessorFromTextureColorSettings(const FTextureSourceColorSettings& InTextureColorSettings, const UE::Color::FColorSpace& InTargetColorSpace);
+
+};
+
 class OPENCOLORIOWRAPPER_API FOpenColorIOProcessorWrapper final
 {
 public:
@@ -119,7 +150,7 @@ public:
 	* @param InDestinationColorSpace Destination color space name.
 	* @param InContextKeyValues (Optional) Additional context modifiers.
 	*/
-	FOpenColorIOProcessorWrapper(
+	explicit FOpenColorIOProcessorWrapper(
 		const FOpenColorIOConfigWrapper* InConfig,
 		FStringView InSourceColorSpace,
 		FStringView InDestinationColorSpace,
@@ -135,7 +166,7 @@ public:
 	* @param bInverseDirection Flag for inverse transform direction.
 	* @param InContextKeyValues (Optional) Additional context modifiers.
 	*/
-	FOpenColorIOProcessorWrapper(
+	explicit FOpenColorIOProcessorWrapper(
 		const FOpenColorIOConfigWrapper* InConfig,
 		FStringView InSourceColorSpace,
 		FStringView InDisplay,
@@ -143,6 +174,20 @@ public:
 		bool bInverseDirection = false,
 		const TMap<FString, FString>& InContextKeyValues = {});
 
+	/**
+	* Constructor.
+	*
+	* @param InConfig Owner config.
+	* @param InNamedTransform Transform name.
+	* @param bInverseDirection Flag for inverse transform direction.
+	* @param InContextKeyValues (Optional) Additional context modifiers.
+	*/
+	explicit FOpenColorIOProcessorWrapper(
+		const FOpenColorIOConfigWrapper* InConfig,
+		FStringView InNamedTransform,
+		bool bInverseDirection = false,
+		const TMap<FString, FString>& InContextKeyValues = {});
+
 	/** Valid when the processor has been successfully created and isn't null. */
 	bool IsValid() const;
 
@@ -167,11 +212,14 @@ public:
 	*
 	* @param InProcessor Parent processor.
 	*/
-	FOpenColorIOCPUProcessorWrapper(FOpenColorIOProcessorWrapper InProcessor);
+	explicit FOpenColorIOCPUProcessorWrapper(FOpenColorIOProcessorWrapper InProcessor);
 
 	/** Valid when the processor has been successfully created and isn't null. */
 	bool IsValid() const;
 
+	/** Apply the color transform in-place to the specified color. */
+	bool TransformColor(FLinearColor& InOutColor) const;
+
 	/** Apply the color transform in-place to the specified image. */
 	bool TransformImage(const FImageView& InOutImage) const;
 
@@ -200,7 +248,7 @@ public:
 	* @param InProcessor Parent processor.
 	* @param InOptions Initialization options.
 	*/
-	FOpenColorIOGPUProcessorWrapper(
+	explicit FOpenColorIOGPUProcessorWrapper(
 		FOpenColorIOProcessorWrapper InProcessor,
 		FInitializationOptions InOptions = FInitializationOptions()
 	);
diff --git a/Engine/Plugins/Compositing/OpenColorIO/Source/OpenColorIOWrapper/Public/OpenColorIOWrapperModule.h b/Engine/Plugins/Compositing/OpenColorIO/Source/OpenColorIOWrapper/Public/OpenColorIOWrapperModule.h
--- a/Engine/Plugins/Compositing/OpenColorIO/Source/OpenColorIOWrapper/Public/OpenColorIOWrapperModule.h
+++ b/Engine/Plugins/Compositing/OpenColorIO/Source/OpenColorIOWrapper/Public/OpenColorIOWrapperModule.h
@@ -5,7 +5,7 @@
 #include "Modules/ModuleManager.h"
 #include "Logging/LogMacros.h"
 
-class FOpenColorIOConfigWrapper;
+class FOpenColorIOEngineBuiltInConfigWrapper;
 
 DECLARE_LOG_CATEGORY_EXTERN(LogOpenColorIOWrapper, Log, All);
 
@@ -39,21 +39,10 @@ public:
 	}
 
 	/**
-	 * Returns a minimal dynamically-created native config for conversions between interchange and working color spaces.
+	 * Returns a minimal dynamically-created native config for engine (working color spaces) conversions.
 	 */
-	virtual const FOpenColorIOConfigWrapper* GetWorkingColorSpaceToInterchangeConfig() const = 0;
-
-	/**
-	* Load a globally-shared config in the module.
-	*
-	* @param InFilePath Config absolute file path.
-	*/
-	virtual void LoadGlobalConfig(FStringView InFilePath) = 0;
-
-	/**
-	* Returns the globally-shared module config if loaded, nullptr otherwise.
-	*/
-	virtual const FOpenColorIOConfigWrapper* GetGlobalConfig() const = 0;
+	virtual FOpenColorIOEngineBuiltInConfigWrapper* GetEngineBuiltInConfig() = 0;
+	virtual const FOpenColorIOEngineBuiltInConfigWrapper* GetEngineBuiltInConfig() const = 0;
 
 	/** Virtual destructor */
 	virtual ~IOpenColorIOWrapperModule() = default;
