OCIO: Fix generated name for transforms from texture settings to working color space to better avoid collisions.

Also removed non-working color space interface since it would complicate this logic needlessly.

#rnx

diff --git a/Engine/Plugins/Compositing/OpenColorIO/Source/OpenColorIOWrapper/Private/OpenColorIOWrapper.cpp b/Engine/Plugins/Compositing/OpenColorIO/Source/OpenColorIOWrapper/Private/OpenColorIOWrapper.cpp
--- a/Engine/Plugins/Compositing/OpenColorIO/Source/OpenColorIOWrapper/Private/OpenColorIOWrapper.cpp
+++ b/Engine/Plugins/Compositing/OpenColorIO/Source/OpenColorIOWrapper/Private/OpenColorIOWrapper.cpp
@@ -394,15 +394,9 @@ FOpenColorIOEngineBuiltInConfigWrapper::FOpenColorIOEngineBuiltInConfigWrapper()
 #endif // WITH_OCIO
 }
 
-FOpenColorIOProcessorWrapper FOpenColorIOEngineBuiltInConfigWrapper::GetProcessorFromTextureColorSettings(const FTextureSourceColorSettings& InTextureColorSettings)
+FOpenColorIOProcessorWrapper FOpenColorIOEngineBuiltInConfigWrapper::GetProcessorToWorkingColorSpace(const FTextureSourceColorSettings& InTextureColorSettings)
 {
-	return GetProcessorFromTextureColorSettings(InTextureColorSettings, UE::Color::FColorSpace::GetWorking());
-}
-
-FOpenColorIOProcessorWrapper FOpenColorIOEngineBuiltInConfigWrapper::GetProcessorFromTextureColorSettings(const FTextureSourceColorSettings& InTextureColorSettings, const UE::Color::FColorSpace& InTargetColorSpace)
-{
-	const uint32 SettingsHash = FCrc::MemCrc32(&InTextureColorSettings, sizeof(FTextureSourceColorSettings));
-	const FString TransformName = FString::Printf(TEXT("WCS%u"), SettingsHash);
+	const FString TransformName = GetTransformToWorkingColorSpaceName(InTextureColorSettings);
 
 #if WITH_OCIO
 	using namespace OCIO_NAMESPACE;
@@ -709,7 +703,7 @@ FOpenColorIOProcessorWrapper FOpenColorIOEngineBuiltInConfigWrapper::GetProcesso
 				InTextureColorSettings.WhiteChromaticityCoordinate
 			);
 			MatrixTransformRcPtr MatrixTransform = MatrixTransform::Create();
-			const FMatrix44d ToWorkingMat = Transpose<double>(FColorSpaceTransform(SourceColorSpace, InTargetColorSpace, ChromaticAdapation));
+			const FMatrix44d ToWorkingMat = Transpose<double>(FColorSpaceTransform(SourceColorSpace, FColorSpace::GetWorking(), ChromaticAdapation));
 			MatrixTransform->setMatrix(&ToWorkingMat.M[0][0]);
 			TransformToWCS->appendTransform(MatrixTransform);
 		}
@@ -717,7 +711,7 @@ FOpenColorIOProcessorWrapper FOpenColorIOEngineBuiltInConfigWrapper::GetProcesso
 		default:
 		{
 			const FColorSpace SourceColorSpace = FColorSpace(static_cast<EColorSpace>(InTextureColorSettings.ColorSpace));
-			const FMatrix44d ToWorkingMat = Transpose<double>(FColorSpaceTransform(SourceColorSpace, InTargetColorSpace, ChromaticAdapation));
+			const FMatrix44d ToWorkingMat = Transpose<double>(FColorSpaceTransform(SourceColorSpace, FColorSpace::GetWorking(), ChromaticAdapation));
 			MatrixTransformRcPtr MatrixTransform = MatrixTransform::Create();
 			MatrixTransform->setMatrix(&ToWorkingMat.M[0][0]);
 			TransformToWCS->appendTransform(MatrixTransform);
@@ -737,6 +731,24 @@ FOpenColorIOProcessorWrapper FOpenColorIOEngineBuiltInConfigWrapper::GetProcesso
 	return FOpenColorIOProcessorWrapper(this, TransformName);
 }
 
+FString FOpenColorIOEngineBuiltInConfigWrapper::GetTransformToWorkingColorSpaceName(const FTextureSourceColorSettings& InTextureColorSettings)
+{
+	const uint32 SettingsId = (uint32)InTextureColorSettings.EncodingOverride | (uint32)InTextureColorSettings.ColorSpace << 8u | (uint32)InTextureColorSettings.ChromaticAdaptationMethod << 16u;
+	FString TransformName = FString::Printf(TEXT("UE_%u"), SettingsId);
+
+	if (InTextureColorSettings.ColorSpace == ETextureColorSpace::TCS_Custom)
+	{
+		uint32 SrcChromaticityHash = 0;
+		SrcChromaticityHash ^= GetTypeHash(InTextureColorSettings.RedChromaticityCoordinate);
+		SrcChromaticityHash ^= GetTypeHash(InTextureColorSettings.GreenChromaticityCoordinate);
+		SrcChromaticityHash ^= GetTypeHash(InTextureColorSettings.BlueChromaticityCoordinate);
+		SrcChromaticityHash ^= GetTypeHash(InTextureColorSettings.WhiteChromaticityCoordinate);
+		TransformName += FString::Printf(TEXT("_%u"), SrcChromaticityHash);
+	}
+
+	return TransformName;
+}
+
 FOpenColorIOProcessorWrapper::FOpenColorIOProcessorWrapper(
 	const FOpenColorIOConfigWrapper* InConfig,
 	FStringView InSourceColorSpace,
diff --git a/Engine/Plugins/Compositing/OpenColorIO/Source/OpenColorIOWrapper/Private/Tests/OpenColorIOWrapperTests.cpp b/Engine/Plugins/Compositing/OpenColorIO/Source/OpenColorIOWrapper/Private/Tests/OpenColorIOWrapperTests.cpp
--- a/Engine/Plugins/Compositing/OpenColorIO/Source/OpenColorIOWrapper/Private/Tests/OpenColorIOWrapperTests.cpp
+++ b/Engine/Plugins/Compositing/OpenColorIO/Source/OpenColorIOWrapper/Private/Tests/OpenColorIOWrapperTests.cpp
@@ -23,7 +23,7 @@ DEFINE_LOG_CATEGORY_STATIC(LogUnrealOpenColorIOTest, Log, All);
 IMPLEMENT_SIMPLE_AUTOMATION_TEST(FOpenColorIOTransferFunctionsTest, "System.OpenColorIO.DecodeToWorkingColorSpace", EAutomationTestFlags::EditorContext | EAutomationTestFlags::EngineFilter)
 bool FOpenColorIOTransferFunctionsTest::RunTest(const FString& Parameters)
 {
-	bool bSuccessful = true;
+	bool bSuccess = true;
 
 #if WITH_OCIO
 	using namespace OCIO_NAMESPACE;
@@ -40,7 +40,7 @@ bool FOpenColorIOTransferFunctionsTest::RunTest(const FString& Parameters)
 		TestSettings.EncodingOverride = static_cast<ETextureSourceEncoding>(TestEncoding);
 		TestSettings.ColorSpace = ETextureColorSpace::TCS_None;
 
-		FOpenColorIOProcessorWrapper Processor = TestConfig.GetProcessorFromTextureColorSettings(TestSettings);
+		FOpenColorIOProcessorWrapper Processor = TestConfig.GetProcessorToWorkingColorSpace(TestSettings);
 		FOpenColorIOCPUProcessorWrapper ProcessorCPU(Processor);
 
 		FLinearColor Actual = TestColor;
@@ -53,12 +53,45 @@ bool FOpenColorIOTransferFunctionsTest::RunTest(const FString& Parameters)
 		{
 			const FString TestNameToPrint = FString::Printf(TEXT("OpenColorIO: %u:%u"), (uint32)TestSettings.EncodingOverride, (uint32)TestSettings.ColorSpace);
 			AddError(FString::Printf(TEXT("Expected '%s' to be %s, but it was %s."), *TestNameToPrint, *Expected.ToString(), *Actual.ToString()), 1);
-			bSuccessful = false;
+			bSuccess = false;
 		}
 	}
+
+
+	// Name collision test
+	int32 Count = 0;
+	TSet<FString> Keys;
+	FRandomStream RandomStream;
+	RandomStream.Initialize(42);
+	for (uint8 TestEncoding = static_cast<uint8>(ETextureSourceEncoding::TSE_None); TestEncoding < static_cast<uint8>(ETextureSourceEncoding::TSE_MAX); ++TestEncoding)
+	{
+		for (uint8 TestColorSpace = static_cast<uint8>(ETextureColorSpace::TCS_None); TestColorSpace < static_cast<uint8>(ETextureColorSpace::TCS_MAX); ++TestColorSpace)
+		{
+			for (uint8 TestChromatic = 0; TestChromatic < 2; ++TestChromatic)
+			{
+				FTextureSourceColorSettings Settings = {};
+				Settings.EncodingOverride = static_cast<ETextureSourceEncoding>(TestEncoding);
+				Settings.ColorSpace = static_cast<ETextureColorSpace>(TestColorSpace);
+				if (Settings.ColorSpace == ETextureColorSpace::TCS_Custom)
+				{
+					Settings.RedChromaticityCoordinate = FVector2D(RandomStream.GetUnitVector());
+					Settings.GreenChromaticityCoordinate = FVector2D(RandomStream.GetUnitVector());
+					Settings.BlueChromaticityCoordinate = FVector2D(RandomStream.GetUnitVector());
+					Settings.WhiteChromaticityCoordinate = FVector2D(RandomStream.GetUnitVector());
+				}
+				Settings.ChromaticAdaptationMethod = static_cast<ETextureChromaticAdaptationMethod>(TestChromatic);
+
+				Keys.Add(FOpenColorIOEngineBuiltInConfigWrapper::GetTransformToWorkingColorSpaceName(Settings));
+				Count++;
+			}
+		}
+	}
+
+	bSuccess &= TestEqual(TEXT("OpenColorIO: Name hash collision test"), Keys.Num(), Count);
+
 #endif
 
-	return bSuccessful;
+	return bSuccess;
 }
 
 #endif
\ No newline at end of file
diff --git a/Engine/Plugins/Compositing/OpenColorIO/Source/OpenColorIOWrapper/Public/OpenColorIOWrapper.h b/Engine/Plugins/Compositing/OpenColorIO/Source/OpenColorIOWrapper/Public/OpenColorIOWrapper.h
--- a/Engine/Plugins/Compositing/OpenColorIO/Source/OpenColorIOWrapper/Public/OpenColorIOWrapper.h
+++ b/Engine/Plugins/Compositing/OpenColorIO/Source/OpenColorIOWrapper/Public/OpenColorIOWrapper.h
@@ -125,18 +125,18 @@ public:
 	FOpenColorIOEngineBuiltInConfigWrapper(FStringView InFilePath, FInitializationOptions InOptions) = delete;
 
 	/**
-	 * Get the processor to working color space from source texture settings.
-	 * @param InTextureColorSettings Source encoding and colorspace.
+	* Get the processor to working color space from source texture settings.
+	* @param InTextureColorSettings Source encoding and colorspace.
+	* @return Processor of the transform
 	*/
-	class FOpenColorIOProcessorWrapper GetProcessorFromTextureColorSettings(const FTextureSourceColorSettings& InTextureColorSettings);
+	class FOpenColorIOProcessorWrapper GetProcessorToWorkingColorSpace(const FTextureSourceColorSettings& InTextureColorSettings);
 
 	/**
-	 * Get the processor to target color space from source texture settings.
-	 * @param InTextureColorSettings Source encoding and colorspace.
-	 * @param InTargetColorSpace Target scene-linear color space.
+	* Get the generated transform name from texture settings to working color space.
+	* @param InTextureColorSettings Source encoding and colorspace.
+	* @return Generated name of the transform
 	*/
-	class FOpenColorIOProcessorWrapper GetProcessorFromTextureColorSettings(const FTextureSourceColorSettings& InTextureColorSettings, const UE::Color::FColorSpace& InTargetColorSpace);
-
+	static FString GetTransformToWorkingColorSpaceName(const FTextureSourceColorSettings& InTextureColorSettings);
 };
 
 class OPENCOLORIOWRAPPER_API FOpenColorIOProcessorWrapper final
