OCIO: Cleanup inteface for engine texture use.

* All classes now prefixed with FOpenColorIOWrapper.
* Removed public link to ImageCore, and public use of engine texture settings.
* Cleaned up constructors and facilited exception catching.
* Exporting methods instead of entire classes.
* Reconnected legacy processor setting, which was accidentally removed.

diff --git a/Engine/Plugins/Compositing/OpenColorIO/Source/OpenColorIO/Private/OpenColorIOColorTransform.cpp b/Engine/Plugins/Compositing/OpenColorIO/Source/OpenColorIO/Private/OpenColorIOColorTransform.cpp
--- a/Engine/Plugins/Compositing/OpenColorIO/Source/OpenColorIO/Private/OpenColorIOColorTransform.cpp
+++ b/Engine/Plugins/Compositing/OpenColorIO/Source/OpenColorIO/Private/OpenColorIOColorTransform.cpp
@@ -26,7 +26,7 @@
 namespace
 {
 	// Returns the (native) config wrapper from the configuration object.
-	FOpenColorIOConfigWrapper* GetTransformConfigWrapper(const UOpenColorIOConfiguration* InConfigurationOwner)
+	FOpenColorIOWrapperConfig* GetTransformConfigWrapper(const UOpenColorIOConfiguration* InConfigurationOwner)
 	{
 		if (IsValid(InConfigurationOwner))
 		{
@@ -37,7 +37,7 @@ namespace
 	}
 
 	// Note: Since OpenColorIO caches processors automatically, we can recreate them without significant additional costs.
-	FOpenColorIOProcessorWrapper GetTransformProcessor(const UOpenColorIOColorTransform* InTransform, const FOpenColorIOConfigWrapper* InConfigWrapper)
+	FOpenColorIOWrapperProcessor GetTransformProcessor(const UOpenColorIOColorTransform* InTransform, const FOpenColorIOWrapperConfig* InConfigWrapper)
 	{
 		check(InTransform);
 		ensure(InConfigWrapper);
@@ -45,7 +45,7 @@ namespace
 		EOpenColorIOViewTransformDirection DisplayViewDirection;
 		if (InTransform->GetDisplayViewDirection(DisplayViewDirection))
 		{
-			return FOpenColorIOProcessorWrapper(
+			return FOpenColorIOWrapperProcessor(
 				InConfigWrapper,
 				InTransform->SourceColorSpace,
 				InTransform->Display,
@@ -55,7 +55,7 @@ namespace
 		}
 		else
 		{
-			return FOpenColorIOProcessorWrapper(
+			return FOpenColorIOWrapperProcessor(
 				InConfigWrapper,
 				InTransform->SourceColorSpace,
 				InTransform->DestinationColorSpace,
@@ -278,15 +278,16 @@ void UOpenColorIOColorTransform::CacheResourceTextures()
 #if WITH_EDITOR
 	if (Textures.IsEmpty())
 	{
+		const UOpenColorIOSettings* Settings = GetDefault<UOpenColorIOSettings>();
 		const UOpenColorIOConfiguration* ConfigurationOwner = Cast<UOpenColorIOConfiguration>(GetOuter());
-		const FOpenColorIOConfigWrapper* ConfigWrapper = GetTransformConfigWrapper(ConfigurationOwner);
+		const FOpenColorIOWrapperConfig* ConfigWrapper = GetTransformConfigWrapper(ConfigurationOwner);
 
 		if (ConfigWrapper != nullptr)
 		{
-			const FOpenColorIOProcessorWrapper TransformProcessor = GetTransformProcessor(this, ConfigWrapper);
+			const FOpenColorIOWrapperProcessor TransformProcessor = GetTransformProcessor(this, ConfigWrapper);
 			if (TransformProcessor.IsValid())
 			{
-				const FOpenColorIOGPUProcessorWrapper GPUProcessor = FOpenColorIOGPUProcessorWrapper(TransformProcessor);
+				const FOpenColorIOWrapperGPUProcessor GPUProcessor = FOpenColorIOWrapperGPUProcessor(TransformProcessor, Settings->bUseLegacyProcessor);
 				const FString ProcessorID = GPUProcessor.GetCacheID();
 
 				//In editor, it will use what's on DDC if there's something corresponding to the actual data or use that raw data
@@ -496,16 +497,14 @@ bool UOpenColorIOColorTransform::IsTransform(const FString& InSourceColorSpace,
 bool UOpenColorIOColorTransform::EditorTransformImage(const FImageView& InOutImage) const
 {
 	const UOpenColorIOConfiguration* ConfigurationOwner = Cast<UOpenColorIOConfiguration>(GetOuter());
-	const FOpenColorIOConfigWrapper* ConfigWrapper = GetTransformConfigWrapper(ConfigurationOwner);
+	const FOpenColorIOWrapperConfig* ConfigWrapper = GetTransformConfigWrapper(ConfigurationOwner);
 
 	if (ConfigWrapper != nullptr)
 	{
-		const FOpenColorIOProcessorWrapper TransformProcessor = GetTransformProcessor(this, ConfigWrapper);
+		const FOpenColorIOWrapperProcessor TransformProcessor = GetTransformProcessor(this, ConfigWrapper);
 		if (TransformProcessor.IsValid())
 		{
-			const FOpenColorIOCPUProcessorWrapper CPUProcessor = FOpenColorIOCPUProcessorWrapper(TransformProcessor);
-
-			return CPUProcessor.TransformImage(InOutImage);
+			return TransformProcessor.TransformImage(InOutImage);
 		}
 	}
 	else
@@ -519,16 +518,14 @@ bool UOpenColorIOColorTransform::EditorTransformImage(const FImageView& InOutIma
 bool UOpenColorIOColorTransform::EditorTransformImage(const FImageView& SrcImage, const FImageView& DestImage) const
 {
 	const UOpenColorIOConfiguration* ConfigurationOwner = Cast<UOpenColorIOConfiguration>(GetOuter());
-	const FOpenColorIOConfigWrapper* ConfigWrapper = GetTransformConfigWrapper(ConfigurationOwner);
+	const FOpenColorIOWrapperConfig* ConfigWrapper = GetTransformConfigWrapper(ConfigurationOwner);
 
 	if (ConfigWrapper != nullptr)
 	{
-		const FOpenColorIOProcessorWrapper TransformProcessor = GetTransformProcessor(this, ConfigWrapper);
+		const FOpenColorIOWrapperProcessor TransformProcessor = GetTransformProcessor(this, ConfigWrapper);
 		if (TransformProcessor.IsValid())
 		{
-			const FOpenColorIOCPUProcessorWrapper CPUProcessor = FOpenColorIOCPUProcessorWrapper(TransformProcessor);
-
-			return CPUProcessor.TransformImage(SrcImage, DestImage);
+			return TransformProcessor.TransformImage(SrcImage, DestImage);
 		}
 	}
 	else
@@ -627,15 +624,16 @@ FString UOpenColorIOColorTransform::GetTransformFriendlyName() const
 bool UOpenColorIOColorTransform::UpdateShaderInfo(FString& OutShaderCodeHash, FString& OutShaderCode, FString& OutRawConfigHash)
 {
 #if WITH_EDITOR
+	const UOpenColorIOSettings* Settings = GetDefault<UOpenColorIOSettings>();
 	const UOpenColorIOConfiguration* ConfigurationOwner = Cast<UOpenColorIOConfiguration>(GetOuter());
-	const FOpenColorIOConfigWrapper* ConfigWrapper = GetTransformConfigWrapper(ConfigurationOwner);
+	const FOpenColorIOWrapperConfig* ConfigWrapper = GetTransformConfigWrapper(ConfigurationOwner);
 
 	if (ConfigWrapper != nullptr)
 	{
-		const FOpenColorIOProcessorWrapper Processor = GetTransformProcessor(this, ConfigWrapper);
+		const FOpenColorIOWrapperProcessor Processor = GetTransformProcessor(this, ConfigWrapper);
 		if (Processor.IsValid())
 		{
-			const FOpenColorIOGPUProcessorWrapper GPUProcessor = FOpenColorIOGPUProcessorWrapper(Processor);
+			const FOpenColorIOWrapperGPUProcessor GPUProcessor = FOpenColorIOWrapperGPUProcessor(Processor, Settings->bUseLegacyProcessor);
 			if (GPUProcessor.IsValid())
 			{
 				OutRawConfigHash = ConfigurationOwner->GetConfigWrapper()->GetCacheID();
diff --git a/Engine/Plugins/Compositing/OpenColorIO/Source/OpenColorIO/Private/OpenColorIOConfiguration.cpp b/Engine/Plugins/Compositing/OpenColorIO/Source/OpenColorIO/Private/OpenColorIOConfiguration.cpp
--- a/Engine/Plugins/Compositing/OpenColorIO/Source/OpenColorIO/Private/OpenColorIOConfiguration.cpp
+++ b/Engine/Plugins/Compositing/OpenColorIO/Source/OpenColorIO/Private/OpenColorIOConfiguration.cpp
@@ -277,7 +277,7 @@ void UOpenColorIOConfiguration::ConfigPathChangedEvent(const TArray<FFileChangeD
 #endif
 }
 
-FOpenColorIOConfigWrapper* UOpenColorIOConfiguration::GetConfigWrapper() const
+FOpenColorIOWrapperConfig* UOpenColorIOConfiguration::GetConfigWrapper() const
 {
 #if WITH_EDITOR
 	return Config.Get();
@@ -613,9 +613,7 @@ void UOpenColorIOConfiguration::LoadConfiguration()
 		FilePath = FPaths::ConvertRelativePathToFull(FPaths::Combine(AbsoluteGameDir, ConfigurationFile.FilePath));
 	}
 
-	FOpenColorIOConfigWrapper::FInitializationOptions Options;
-	Options.bAddWorkingColorSpace = true;
-	Config = MakePimpl<FOpenColorIOConfigWrapper>(FilePath, Options);
+	Config = MakePimpl<FOpenColorIOWrapperConfig>(FilePath, FOpenColorIOWrapperConfig::WCS_AsInterchangeSpace);
 
 	if (Config->IsValid())
 	{
diff --git a/Engine/Plugins/Compositing/OpenColorIO/Source/OpenColorIO/Public/IOpenColorIOModule.h b/Engine/Plugins/Compositing/OpenColorIO/Source/OpenColorIO/Public/IOpenColorIOModule.h
--- a/Engine/Plugins/Compositing/OpenColorIO/Source/OpenColorIO/Public/IOpenColorIOModule.h
+++ b/Engine/Plugins/Compositing/OpenColorIO/Source/OpenColorIO/Public/IOpenColorIOModule.h
@@ -45,6 +45,6 @@ public:
 	 */
 	virtual FOpenColorIODisplayManager& GetDisplayManager() = 0;
 
-	UE_DEPRECATED(5.3, "This function is deprecated and has been replaced by FOpenColorIOConfigWrapper class in the OpenColorIOWrapper module.")
+	UE_DEPRECATED(5.3, "This function is deprecated and has been replaced by FOpenColorIOWrapperConfig class in the OpenColorIOWrapper module.")
 	FOpenColorIONativeInterchangeConfiguration* GetNativeInterchangeConfig_Internal() { return nullptr; };
 };
\ No newline at end of file
diff --git a/Engine/Plugins/Compositing/OpenColorIO/Source/OpenColorIO/Public/OpenColorIOConfiguration.h b/Engine/Plugins/Compositing/OpenColorIO/Source/OpenColorIO/Public/OpenColorIOConfiguration.h
--- a/Engine/Plugins/Compositing/OpenColorIO/Source/OpenColorIO/Public/OpenColorIOConfiguration.h
+++ b/Engine/Plugins/Compositing/OpenColorIO/Source/OpenColorIO/Public/OpenColorIOConfiguration.h
@@ -14,7 +14,7 @@
 #include "OpenColorIOConfiguration.generated.h"
 
 
-class FOpenColorIOConfigWrapper;
+class FOpenColorIOWrapperConfig;
 class FOpenColorIOTransformResource;
 class FTextureResource;
 class UOpenColorIOColorTransform;
@@ -64,7 +64,7 @@ public:
 	void ConfigPathChangedEvent(const TArray<FFileChangeData>& InFileChanges, const FString InFileMountPath);
 
 	/** Internal only: Replacement for previous `GetNativeConfig_Internal()`, `GetLoadedConfiguration()` and `GetLoadedConfigurationFile()` functions, returning the private implementation of the native OCIO config. */
-	FOpenColorIOConfigWrapper* GetConfigWrapper() const;
+	FOpenColorIOWrapperConfig* GetConfigWrapper() const;
 
 protected:
 
@@ -137,7 +137,7 @@ private:
 
 #if WITH_EDITORONLY_DATA
 	/** Private implementation of the OpenColorIO config object. */
-	TPimplPtr<FOpenColorIOConfigWrapper> Config = nullptr;
+	TPimplPtr<FOpenColorIOWrapperConfig> Config = nullptr;
 #endif
 };
 
@@ -146,28 +146,28 @@ private:
 class FOpenColorIOEditorConfigurationInspector {
 public:
 
-	UE_DEPRECATED(5.3, "This class is deprecated and has been replaced by FOpenColorIOConfigWrapper class in the OpenColorIOWrapper module.")
+	UE_DEPRECATED(5.3, "This class is deprecated and has been replaced by FOpenColorIOWrapperConfig class in the OpenColorIOWrapper module.")
 	OPENCOLORIO_API FOpenColorIOEditorConfigurationInspector(const UOpenColorIOConfiguration& InConfiguration);
 
-	UE_DEPRECATED(5.3, "This class is deprecated and has been replaced by FOpenColorIOConfigWrapper class in the OpenColorIOWrapper module.")
+	UE_DEPRECATED(5.3, "This class is deprecated and has been replaced by FOpenColorIOWrapperConfig class in the OpenColorIOWrapper module.")
 	OPENCOLORIO_API int32 GetNumColorSpaces() const;
 	
-	UE_DEPRECATED(5.3, "This class is deprecated and has been replaced by FOpenColorIOConfigWrapper class in the OpenColorIOWrapper module.")
+	UE_DEPRECATED(5.3, "This class is deprecated and has been replaced by FOpenColorIOWrapperConfig class in the OpenColorIOWrapper module.")
 	OPENCOLORIO_API FString GetColorSpaceName(int32 Index) const;
 	
-	UE_DEPRECATED(5.3, "This class is deprecated and has been replaced by FOpenColorIOConfigWrapper class in the OpenColorIOWrapper module.")
+	UE_DEPRECATED(5.3, "This class is deprecated and has been replaced by FOpenColorIOWrapperConfig class in the OpenColorIOWrapper module.")
 	OPENCOLORIO_API FString GetColorSpaceFamilyName(const TCHAR* InColorSpaceName) const;
 
-	UE_DEPRECATED(5.3, "This class is deprecated and has been replaced by FOpenColorIOConfigWrapper class in the OpenColorIOWrapper module.")
+	UE_DEPRECATED(5.3, "This class is deprecated and has been replaced by FOpenColorIOWrapperConfig class in the OpenColorIOWrapper module.")
 	OPENCOLORIO_API int32 GetNumDisplays() const;
 	
-	UE_DEPRECATED(5.3, "This class is deprecated and has been replaced by FOpenColorIOConfigWrapper class in the OpenColorIOWrapper module.")
+	UE_DEPRECATED(5.3, "This class is deprecated and has been replaced by FOpenColorIOWrapperConfig class in the OpenColorIOWrapper module.")
 	OPENCOLORIO_API FString GetDisplayName(int32 Index) const;
 
-	UE_DEPRECATED(5.3, "This class is deprecated and has been replaced by FOpenColorIOConfigWrapper class in the OpenColorIOWrapper module.")
+	UE_DEPRECATED(5.3, "This class is deprecated and has been replaced by FOpenColorIOWrapperConfig class in the OpenColorIOWrapper module.")
 	OPENCOLORIO_API int32 GetNumViews(const TCHAR* InDisplayName) const;
 	
-	UE_DEPRECATED(5.3, "This class is deprecated and has been replaced by FOpenColorIOConfigWrapper class in the OpenColorIOWrapper module.")
+	UE_DEPRECATED(5.3, "This class is deprecated and has been replaced by FOpenColorIOWrapperConfig class in the OpenColorIOWrapper module.")
 	OPENCOLORIO_API FString GetViewName(const TCHAR* InDisplayName, int32 Index) const;
 
 private:
diff --git a/Engine/Plugins/Compositing/OpenColorIO/Source/OpenColorIOEditor/Private/OpenColorIOColorSpaceCustomization.cpp b/Engine/Plugins/Compositing/OpenColorIO/Source/OpenColorIOEditor/Private/OpenColorIOColorSpaceCustomization.cpp
--- a/Engine/Plugins/Compositing/OpenColorIO/Source/OpenColorIOEditor/Private/OpenColorIOColorSpaceCustomization.cpp
+++ b/Engine/Plugins/Compositing/OpenColorIO/Source/OpenColorIOEditor/Private/OpenColorIOColorSpaceCustomization.cpp
@@ -113,7 +113,7 @@ IPropertyTypeCustomizationOpenColorIO::IPropertyTypeCustomizationOpenColorIO(TSh
 {
 }
 
-FOpenColorIOConfigWrapper* IPropertyTypeCustomizationOpenColorIO::GetConfigWrapper(const TSharedPtr<IPropertyHandle>& InConfigurationAssetProperty)
+FOpenColorIOWrapperConfig* IPropertyTypeCustomizationOpenColorIO::GetConfigWrapper(const TSharedPtr<IPropertyHandle>& InConfigurationAssetProperty)
 {
 	if (InConfigurationAssetProperty.IsValid())
 	{
@@ -165,7 +165,7 @@ void FOpenColorIOColorSpaceCustomization::PopulateSubMenu(FMenuBuilder& InMenuBu
 	// To keep track of submenus that were already added
 	TArray<FString> ExistingSubMenus;
 
-	const FOpenColorIOConfigWrapper* ConfigWrapper = GetConfigWrapper(ConfigurationObjectProperty);
+	const FOpenColorIOWrapperConfig* ConfigWrapper = GetConfigWrapper(ConfigurationObjectProperty);
 
 	if (ConfigWrapper)
 	{
@@ -223,7 +223,7 @@ void FOpenColorIOColorSpaceCustomization::AddMenuEntry(FMenuBuilder& InMenuBuild
 
 TSharedRef<SWidget> FOpenColorIOColorSpaceCustomization::HandleSourceComboButtonMenuContent()
 {
-	const FOpenColorIOConfigWrapper* ConfigWrapper = GetConfigWrapper(ConfigurationObjectProperty);
+	const FOpenColorIOWrapperConfig* ConfigWrapper = GetConfigWrapper(ConfigurationObjectProperty);
 
 	// Generate menu
 	FMenuBuilder MenuBuilder(true, nullptr);
@@ -265,7 +265,7 @@ TSharedRef<SWidget> FOpenColorIOColorSpaceCustomization::HandleSourceComboButton
 
 void FOpenColorIODisplayViewCustomization::PopulateViewSubMenu(FMenuBuilder& InMenuBuilder, FOpenColorIODisplayView InDisplayView)
 {
-	const FOpenColorIOConfigWrapper* ConfigWrapper = GetConfigWrapper(ConfigurationObjectProperty);
+	const FOpenColorIOWrapperConfig* ConfigWrapper = GetConfigWrapper(ConfigurationObjectProperty);
 
 	if (ConfigWrapper)
 	{
@@ -315,7 +315,7 @@ void FOpenColorIODisplayViewCustomization::AddMenuEntry(FMenuBuilder& InMenuBuil
 
 TSharedRef<SWidget> FOpenColorIODisplayViewCustomization::HandleSourceComboButtonMenuContent()
 {
-	const FOpenColorIOConfigWrapper* ConfigWrapper = GetConfigWrapper(ConfigurationObjectProperty);
+	const FOpenColorIOWrapperConfig* ConfigWrapper = GetConfigWrapper(ConfigurationObjectProperty);
 
 	// Generate menu
 	FMenuBuilder MenuBuilder(true, nullptr);
diff --git a/Engine/Plugins/Compositing/OpenColorIO/Source/OpenColorIOEditor/Private/OpenColorIOColorSpaceCustomization.h b/Engine/Plugins/Compositing/OpenColorIO/Source/OpenColorIOEditor/Private/OpenColorIOColorSpaceCustomization.h
--- a/Engine/Plugins/Compositing/OpenColorIO/Source/OpenColorIOEditor/Private/OpenColorIOColorSpaceCustomization.h
+++ b/Engine/Plugins/Compositing/OpenColorIO/Source/OpenColorIOEditor/Private/OpenColorIOColorSpaceCustomization.h
@@ -9,7 +9,7 @@
 #include "OpenColorIOColorSpace.h"
 #include "Widgets/SWidget.h"
 
-class FOpenColorIOConfigWrapper;
+class FOpenColorIOWrapperConfig;
 
 /**
  * Implements a details view customization for the FOpenColorIOConfiguration
@@ -25,7 +25,7 @@ public:
 
 protected:
 
-	static FOpenColorIOConfigWrapper* GetConfigWrapper(const TSharedPtr<IPropertyHandle>& InConfigurationObjectProperty);
+	static FOpenColorIOWrapperConfig* GetConfigWrapper(const TSharedPtr<IPropertyHandle>& InConfigurationObjectProperty);
 
 protected:
 
diff --git a/Engine/Plugins/Compositing/OpenColorIO/Source/OpenColorIOWrapper/OpenColorIOWrapper.Build.cs b/Engine/Plugins/Compositing/OpenColorIO/Source/OpenColorIOWrapper/OpenColorIOWrapper.Build.cs
--- a/Engine/Plugins/Compositing/OpenColorIO/Source/OpenColorIOWrapper/OpenColorIOWrapper.Build.cs
+++ b/Engine/Plugins/Compositing/OpenColorIO/Source/OpenColorIOWrapper/OpenColorIOWrapper.Build.cs
@@ -8,6 +8,11 @@ namespace UnrealBuildTool.Rules
 		{
 			PCHUsage = PCHUsageMode.UseExplicitOrSharedPCHs;
 
+			PublicIncludePathModuleNames.AddRange(new string[]
+			{
+				"ColorManagement", // for ColorManagementDefines.h
+			});
+
 			PrivateIncludePathModuleNames.AddRange(new string[]
 			{
 				"Engine", // for TextureDefines.h
@@ -16,6 +21,7 @@ namespace UnrealBuildTool.Rules
 			PrivateDependencyModuleNames.AddRange(new string[]
 			{
 				"ColorManagement",
+				"ImageCore",
 			});
 
 			if (Target.bBuildEditor)
@@ -30,7 +36,6 @@ namespace UnrealBuildTool.Rules
 			PublicDependencyModuleNames.AddRange(new string[]
 			{
 				"Core",
-				"ImageCore"
 			});
 		}
 	}
diff --git a/Engine/Plugins/Compositing/OpenColorIO/Source/OpenColorIOWrapper/Private/OpenColorIOWrapper.cpp b/Engine/Plugins/Compositing/OpenColorIO/Source/OpenColorIOWrapper/Private/OpenColorIOWrapper.cpp
--- a/Engine/Plugins/Compositing/OpenColorIO/Source/OpenColorIOWrapper/Private/OpenColorIOWrapper.cpp
+++ b/Engine/Plugins/Compositing/OpenColorIO/Source/OpenColorIOWrapper/Private/OpenColorIOWrapper.cpp
@@ -3,9 +3,11 @@
 #include "OpenColorIOWrapper.h"
 #include "OpenColorIOWrapperModule.h"
 
+#include "Containers/StringView.h"
+#include "Containers/UnrealString.h"
 #include "Containers/Map.h"
 #include "ColorSpace.h"
-#include "Engine/Texture.h"
+#include "Engine/TextureDefines.h"
 #include "ImageCore.h"
 
 #if WITH_OCIO
@@ -16,6 +18,20 @@ THIRD_PARTY_INCLUDES_END
 
 #define LOCTEXT_NAMESPACE "OpenColorIOWrapper"
 
+#if PLATFORM_EXCEPTIONS_DISABLED
+	#define OCIO_EXCEPTION_HANDLING_TRY()
+	#define OCIO_EXCEPTION_HANDLING_CATCH(Verbosity, Format, ...) (0)
+#else
+	#define OCIO_EXCEPTION_HANDLING_TRY() \
+		try {
+	
+	// NOTE: A string parameter must always be added to the format string since we automatically include the error message.
+	#define OCIO_EXCEPTION_HANDLING_CATCH(Verbosity, Format, ...) \
+		} catch (OCIO_NAMESPACE::Exception& Exc) { \
+			UE_LOG(LogOpenColorIOWrapper, Verbosity, Format, ##__VA_ARGS__, StringCast<TCHAR>(Exc.what()).Get()); \
+		}
+#endif
+
 namespace OpenColorIOWrapper
 {
 	/** OCIO_NAMESPACE::ROLE_INTERCHANGE_SCENE equivalent, since we currently cannot delay-load this definition. */
@@ -120,57 +136,57 @@ namespace {
 #endif // WITH_OCIO
 
 
-FOpenColorIOConfigWrapper::FOpenColorIOConfigWrapper()
+FOpenColorIOWrapperConfig::FOpenColorIOWrapperConfig()
 	: Pimpl(MakePimpl<FOpenColorIOConfigPimpl, EPimplPtrMode::DeepCopy>())
 { }
 
-FOpenColorIOConfigWrapper::FOpenColorIOConfigWrapper(FStringView InFilePath, FInitializationOptions InOptions)
-	: FOpenColorIOConfigWrapper()
+FOpenColorIOWrapperConfig::FOpenColorIOWrapperConfig(FStringView InFilePath, EAddWorkingColorSpaceOption InOption)
+	: FOpenColorIOWrapperConfig()
 {
 #if WITH_OCIO
-	TRACE_CPUPROFILER_EVENT_SCOPE(FOpenColorIOConfigWrapper::FOpenColorIOConfigWrapper)
-#if !PLATFORM_EXCEPTIONS_DISABLED
-	try
-#endif
-	{
+	TRACE_CPUPROFILER_EVENT_SCOPE(FOpenColorIOWrapperConfig::FOpenColorIOWrapperConfig)
+	OCIO_EXCEPTION_HANDLING_TRY()
+
 		using namespace OCIO_NAMESPACE;
 
 		ConstConfigRcPtr NewConfig = Config::CreateFromFile(StringCast<ANSICHAR>(InFilePath.GetData()).Get());
 
-		if (NewConfig && InOptions.bAddWorkingColorSpace)
+		if (NewConfig != nullptr)
 		{
-			const TUniquePtr<ANSICHAR[]> AnsiWorkingColorSpaceName = OpenColorIOWrapper::MakeAnsiString(OpenColorIOWrapper::GetWorkingColorSpaceName());
-			ConstColorSpaceRcPtr InterchangeCS = NewConfig->getColorSpace(NewConfig->getCanonicalName(OpenColorIOWrapper::GetInterchangeName()));
-
-			// When the aces interchange color space is present, we add the working color space as an additional option.
-			if (InterchangeCS != nullptr && NewConfig->getColorSpace(AnsiWorkingColorSpaceName.Get()) == nullptr)
+			if (InOption == WCS_AsInterchangeSpace)
 			{
-				ColorSpaceRcPtr WorkingCS = InterchangeCS->createEditableCopy();
-				WorkingCS->setName(AnsiWorkingColorSpaceName.Get());
-				WorkingCS->setFamily("UE");
-				WorkingCS->clearAliases();
+				const TUniquePtr<ANSICHAR[]> AnsiWorkingColorSpaceName = OpenColorIOWrapper::MakeAnsiString(OpenColorIOWrapper::GetWorkingColorSpaceName());
+				ConstColorSpaceRcPtr InterchangeCS = NewConfig->getColorSpace(NewConfig->getCanonicalName(OpenColorIOWrapper::GetInterchangeName()));
 
-				ConfigRcPtr NewConfigCopy = NewConfig->createEditableCopy();
-				NewConfigCopy->addColorSpace(WorkingCS);
+				// When the aces interchange color space is present, we add the working color space as an additional option.
+				if (InterchangeCS != nullptr && NewConfig->getColorSpace(AnsiWorkingColorSpaceName.Get()) == nullptr)
+				{
+					ColorSpaceRcPtr WorkingCS = InterchangeCS->createEditableCopy();
+					WorkingCS->setName(AnsiWorkingColorSpaceName.Get());
+					WorkingCS->setFamily("UE");
+					WorkingCS->clearAliases();
+
+					ConfigRcPtr NewConfigCopy = NewConfig->createEditableCopy();
+					NewConfigCopy->addColorSpace(WorkingCS);
 
-				NewConfig = NewConfigCopy;
+					NewConfig = NewConfigCopy;
+				}
+			}
+			else
+			{
+				UE_LOG(LogOpenColorIOWrapper, Log, TEXT("Could not add the working color space since the config does not define \"aces_interchange\"."))
 			}
 		}
 
 		Pimpl->Config = NewConfig;
-	}
-#if !PLATFORM_EXCEPTIONS_DISABLED
-	catch (OCIO_NAMESPACE::Exception& Exc)
-	{
-		UE_LOG(LogOpenColorIOWrapper, Error, TEXT("Could not create OCIO configuration file for %s. Error message: %s."), InFilePath.GetData(), StringCast<TCHAR>(Exc.what()).Get());
-	}
-#endif
+	
+	OCIO_EXCEPTION_HANDLING_CATCH(Error, TEXT("Could not create OCIO configuration file for %s. Error message: %s"), InFilePath.GetData())
 #else
 	UE_LOG(LogOpenColorIOWrapper, Error, TEXT("The OpenColorIO library is not available, functionality will be disabled."));
 #endif // WITH_OCIO
 }
 
-bool FOpenColorIOConfigWrapper::IsValid() const
+bool FOpenColorIOWrapperConfig::IsValid() const
 {
 #if WITH_OCIO
 	return Pimpl->Config != nullptr;
@@ -179,7 +195,7 @@ bool FOpenColorIOConfigWrapper::IsValid() const
 #endif // WITH_OCIO
 }
 
-int32 FOpenColorIOConfigWrapper::GetNumColorSpaces() const
+int32 FOpenColorIOWrapperConfig::GetNumColorSpaces() const
 {
 #if WITH_OCIO
 	if (IsValid())
@@ -191,7 +207,7 @@ int32 FOpenColorIOConfigWrapper::GetNumColorSpaces() const
 	return 0;
 }
 
-FString FOpenColorIOConfigWrapper::GetColorSpaceName(int32 Index) const
+FString FOpenColorIOWrapperConfig::GetColorSpaceName(int32 Index) const
 {
 #if WITH_OCIO
 	if (IsValid())
@@ -205,7 +221,7 @@ FString FOpenColorIOConfigWrapper::GetColorSpaceName(int32 Index) const
 	return {};
 }
 
-int32 FOpenColorIOConfigWrapper::GetColorSpaceIndex(const TCHAR* InColorSpaceName)
+int32 FOpenColorIOWrapperConfig::GetColorSpaceIndex(const TCHAR* InColorSpaceName)
 {
 #if WITH_OCIO
 	if (IsValid())
@@ -217,7 +233,7 @@ int32 FOpenColorIOConfigWrapper::GetColorSpaceIndex(const TCHAR* InColorSpaceNam
 	return false;
 }
 
-FString FOpenColorIOConfigWrapper::GetColorSpaceFamilyName(const TCHAR* InColorSpaceName) const
+FString FOpenColorIOWrapperConfig::GetColorSpaceFamilyName(const TCHAR* InColorSpaceName) const
 {
 #if WITH_OCIO
 	if (IsValid())
@@ -234,7 +250,7 @@ FString FOpenColorIOConfigWrapper::GetColorSpaceFamilyName(const TCHAR* InColorS
 	return {};
 }
 
-int32 FOpenColorIOConfigWrapper::GetNumDisplays() const
+int32 FOpenColorIOWrapperConfig::GetNumDisplays() const
 {
 #if WITH_OCIO
 	if (IsValid())
@@ -246,7 +262,7 @@ int32 FOpenColorIOConfigWrapper::GetNumDisplays() const
 	return 0;
 }
 
-FString FOpenColorIOConfigWrapper::GetDisplayName(int32 Index) const
+FString FOpenColorIOWrapperConfig::GetDisplayName(int32 Index) const
 {
 #if WITH_OCIO
 	if (IsValid())
@@ -260,7 +276,7 @@ FString FOpenColorIOConfigWrapper::GetDisplayName(int32 Index) const
 	return {};
 }
 
-int32 FOpenColorIOConfigWrapper::GetNumViews(const TCHAR* InDisplayName) const
+int32 FOpenColorIOWrapperConfig::GetNumViews(const TCHAR* InDisplayName) const
 {
 #if WITH_OCIO
 	if (IsValid())
@@ -272,7 +288,7 @@ int32 FOpenColorIOConfigWrapper::GetNumViews(const TCHAR* InDisplayName) const
 	return 0;
 }
 
-FString FOpenColorIOConfigWrapper::GetViewName(const TCHAR* InDisplayName, int32 Index) const
+FString FOpenColorIOWrapperConfig::GetViewName(const TCHAR* InDisplayName, int32 Index) const
 {
 #if WITH_OCIO
 	if (IsValid())
@@ -286,7 +302,7 @@ FString FOpenColorIOConfigWrapper::GetViewName(const TCHAR* InDisplayName, int32
 	return {};
 }
 
-FString FOpenColorIOConfigWrapper::GetDisplayViewTransformName(const TCHAR* InDisplayName, const TCHAR* InViewName) const
+FString FOpenColorIOWrapperConfig::GetDisplayViewTransformName(const TCHAR* InDisplayName, const TCHAR* InViewName) const
 {
 #if WITH_OCIO
 	if (IsValid())
@@ -300,7 +316,7 @@ FString FOpenColorIOConfigWrapper::GetDisplayViewTransformName(const TCHAR* InDi
 	return {};
 }
 
-FString FOpenColorIOConfigWrapper::GetCacheID() const
+FString FOpenColorIOWrapperConfig::GetCacheID() const
 {
 #if WITH_OCIO
 	if (IsValid())
@@ -312,7 +328,7 @@ FString FOpenColorIOConfigWrapper::GetCacheID() const
 	return {};
 }
 
-FString FOpenColorIOConfigWrapper::GetDebugString() const
+FString FOpenColorIOWrapperConfig::GetDebugString() const
 {
 	TStringBuilder<1024> DebugStringBuilder;
 	
@@ -368,12 +384,13 @@ FString FOpenColorIOConfigWrapper::GetDebugString() const
 	return DebugStringBuilder.ToString();
 }
 
-FOpenColorIOEngineBuiltInConfigWrapper::FOpenColorIOEngineBuiltInConfigWrapper()
-	: FOpenColorIOConfigWrapper()
+
+FOpenColorIOWrapperEngineConfig::FOpenColorIOWrapperEngineConfig()
+	: FOpenColorIOWrapperConfig()
 {
 #if WITH_OCIO
-	TRACE_CPUPROFILER_EVENT_SCOPE(FOpenColorIOConfigWrapper::CreateWorkingColorSpaceToInterchangeConfig)
-		using namespace OCIO_NAMESPACE;
+	TRACE_CPUPROFILER_EVENT_SCOPE(FOpenColorIOWrapperConfig::CreateWorkingColorSpaceToInterchangeConfig)
+	using namespace OCIO_NAMESPACE;
 	using namespace UE::Color;
 
 	ConfigRcPtr StudioConfig = Config::CreateFromBuiltinConfig("studio-config-v1.0.0_aces-v1.3_ocio-v2.1")->createEditableCopy();
@@ -391,30 +408,146 @@ FOpenColorIOEngineBuiltInConfigWrapper::FOpenColorIOEngineBuiltInConfigWrapper()
 	StudioConfig->addColorSpace(WCS);
 
 	Pimpl->Config = StudioConfig;
+
+	check(IsValid());
 #endif // WITH_OCIO
 }
 
-FOpenColorIOProcessorWrapper FOpenColorIOEngineBuiltInConfigWrapper::GetProcessorToWorkingColorSpace(const FTextureSourceColorSettings& InTextureColorSettings)
+FOpenColorIOWrapperProcessor::FOpenColorIOWrapperProcessor(
+	const FOpenColorIOWrapperConfig* InConfig,
+	FStringView InSourceColorSpace,
+	FStringView InDestinationColorSpace,
+	const TMap<FString, FString>& InContextKeyValues)
+	: Pimpl(MakePimpl<FOpenColorIOProcessorPimpl, EPimplPtrMode::DeepCopy>())
+	, OwnerConfig(InConfig)
+	, WorkingColorSpaceTransformType(EOpenColorIOWorkingColorSpaceTransform::None)
 {
-	const FString TransformName = GetTransformToWorkingColorSpaceName(InTextureColorSettings);
+	if (InSourceColorSpace == OpenColorIOWrapper::GetWorkingColorSpaceName())
+	{
+		WorkingColorSpaceTransformType = EOpenColorIOWorkingColorSpaceTransform::Source;
+	}
+	else if (InDestinationColorSpace == OpenColorIOWrapper::GetWorkingColorSpaceName())
+	{
+		WorkingColorSpaceTransformType = EOpenColorIOWorkingColorSpaceTransform::Destination;
+	}
+
+#if WITH_OCIO
+	OCIO_EXCEPTION_HANDLING_TRY()
+		if (OwnerConfig != nullptr && OwnerConfig->IsValid())
+		{
+			const OCIO_NAMESPACE::ConstConfigRcPtr& Config = OwnerConfig->Pimpl->Config;
+			OCIO_NAMESPACE::ContextRcPtr Context = Config->getCurrentContext()->createEditableCopy();
+
+			for (const TPair<FString, FString>& KeyValue : InContextKeyValues)
+			{
+				Context->setStringVar(TCHAR_TO_ANSI(*KeyValue.Key), TCHAR_TO_ANSI(*KeyValue.Value));
+			}
+
+			Pimpl->Processor = Config->getProcessor(
+				Context,
+				StringCast<ANSICHAR>(InSourceColorSpace.GetData()).Get(),
+				StringCast<ANSICHAR>(InDestinationColorSpace.GetData()).Get()
+			);
+		}
+	OCIO_EXCEPTION_HANDLING_CATCH(Log, TEXT("Failed to create processor for [%s, %s]. Error message: %s"), InSourceColorSpace.GetData(), InDestinationColorSpace.GetData())
+#endif // WITH_OCIO
+}
+
+FOpenColorIOWrapperProcessor::FOpenColorIOWrapperProcessor(
+	const FOpenColorIOWrapperConfig* InConfig,
+	FStringView InSourceColorSpace,
+	FStringView InDisplay,
+	FStringView InView,
+	bool bInverseDirection,
+	const TMap<FString, FString>& InContextKeyValues)
+	: Pimpl(MakePimpl<FOpenColorIOProcessorPimpl, EPimplPtrMode::DeepCopy>())
+	, OwnerConfig(InConfig)
+	, WorkingColorSpaceTransformType(EOpenColorIOWorkingColorSpaceTransform::None)
+{
+	if (InSourceColorSpace == OpenColorIOWrapper::GetWorkingColorSpaceName())
+	{
+		if (bInverseDirection)
+		{
+			WorkingColorSpaceTransformType = EOpenColorIOWorkingColorSpaceTransform::Destination;
+		}
+		else
+		{
+			WorkingColorSpaceTransformType = EOpenColorIOWorkingColorSpaceTransform::Source;
+		}
+	}
+
+#if WITH_OCIO
+	OCIO_EXCEPTION_HANDLING_TRY()
+		if (OwnerConfig != nullptr && OwnerConfig->IsValid())
+		{
+			const OCIO_NAMESPACE::ConstConfigRcPtr& Config = OwnerConfig->Pimpl->Config;
+			OCIO_NAMESPACE::ContextRcPtr Context = Config->getCurrentContext()->createEditableCopy();
+
+			for (const TPair<FString, FString>& KeyValue : InContextKeyValues)
+			{
+				Context->setStringVar(TCHAR_TO_ANSI(*KeyValue.Key), TCHAR_TO_ANSI(*KeyValue.Value));
+			}
+
+			Pimpl->Processor = Config->getProcessor(
+				StringCast<ANSICHAR>(InSourceColorSpace.GetData()).Get(),
+				StringCast<ANSICHAR>(InDisplay.GetData()).Get(),
+				StringCast<ANSICHAR>(InView.GetData()).Get(),
+				static_cast<OCIO_NAMESPACE::TransformDirection>(bInverseDirection)
+			);
+		}
+	OCIO_EXCEPTION_HANDLING_CATCH(Log, TEXT("Failed to create processor for [%s, %s, %s, %s]. Error message: %s"), InSourceColorSpace.GetData(), InDisplay.GetData(), InView.GetData(), (bInverseDirection ? TEXT("Inverse") : TEXT("Forward")));
+#endif // WITH_OCIO
+}
+
+FOpenColorIOWrapperProcessor::FOpenColorIOWrapperProcessor(const FOpenColorIOWrapperConfig* InConfig, FStringView InNamedTransform, bool bInverseDirection, const TMap<FString, FString>& InContextKeyValues)
+	: Pimpl(MakePimpl<FOpenColorIOProcessorPimpl, EPimplPtrMode::DeepCopy>())
+	, OwnerConfig(InConfig)
+	, WorkingColorSpaceTransformType(EOpenColorIOWorkingColorSpaceTransform::None)
+{
+#if WITH_OCIO
+	OCIO_EXCEPTION_HANDLING_TRY()
+		if (OwnerConfig != nullptr && OwnerConfig->IsValid())
+		{
+			const OCIO_NAMESPACE::ConstConfigRcPtr& Config = OwnerConfig->Pimpl->Config;
+			OCIO_NAMESPACE::ContextRcPtr Context = Config->getCurrentContext()->createEditableCopy();
+
+			for (const TPair<FString, FString>& KeyValue : InContextKeyValues)
+			{
+				Context->setStringVar(TCHAR_TO_ANSI(*KeyValue.Key), TCHAR_TO_ANSI(*KeyValue.Value));
+			}
+
+			Pimpl->Processor = Config->getProcessor(
+				StringCast<ANSICHAR>(InNamedTransform.GetData()).Get(),
+				static_cast<OCIO_NAMESPACE::TransformDirection>(bInverseDirection)
+			);
+		}
+	OCIO_EXCEPTION_HANDLING_CATCH(Log, TEXT("Failed to create processor for [%s]. Error message: %s"), InNamedTransform.GetData());
+#endif // WITH_OCIO
+}
+
+FOpenColorIOWrapperProcessor FOpenColorIOWrapperProcessor::CreateTransformToWorkingColorSpace(const FOpenColorIOWrapperSourceColorSettings& InColorSettings)
+{
+	FOpenColorIOWrapperConfig& BuiltInConfig = IOpenColorIOWrapperModule::Get().GetEngineBuiltInConfig();
+
+	const FString TransformName = GetTransformToWorkingColorSpaceName(InColorSettings);
 
 #if WITH_OCIO
 	using namespace OCIO_NAMESPACE;
 	using namespace UE::Color;
 
-	if (Pimpl->Config->getNamedTransform(StringCast<ANSICHAR>(*TransformName).Get()) == nullptr)
+	if (BuiltInConfig.Pimpl->Config->getNamedTransform(StringCast<ANSICHAR>(*TransformName).Get()) == nullptr)
 	{
 		GroupTransformRcPtr TransformToWCS = GroupTransform::Create();
 		NamedTransformRcPtr ParentTransform = NamedTransform::Create();
 		ParentTransform->setName(StringCast<ANSICHAR>(*TransformName).Get());
 
-		switch (InTextureColorSettings.EncodingOverride)
+		switch (InColorSettings.EncodingOverride)
 		{
-		case ETextureSourceEncoding::TSE_None:
-		case ETextureSourceEncoding::TSE_Linear:
+		case EEncoding::None:
+		case EEncoding::Linear:
 			ParentTransform->setEncoding("scene-linear");
 			break;
-		case ETextureSourceEncoding::TSE_sRGB:
+		case EEncoding::sRGB:
 		{
 			ExponentWithLinearTransformRcPtr ChildTransform = ExponentWithLinearTransform::Create();
 			ChildTransform->setGamma({ 2.4, 2.4, 2.4, 1.0 });
@@ -424,7 +557,7 @@ FOpenColorIOProcessorWrapper FOpenColorIOEngineBuiltInConfigWrapper::GetProcesso
 			ParentTransform->setEncoding("sdr-video");
 		}
 		break;
-		case ETextureSourceEncoding::TSE_ST2084:
+		case EEncoding::ST2084:
 		{
 			BuiltinTransformRcPtr ChildTransform = BuiltinTransform::Create();
 			ChildTransform->setStyle("CURVE - ST-2084_to_LINEAR");
@@ -439,7 +572,7 @@ FOpenColorIOProcessorWrapper FOpenColorIOEngineBuiltInConfigWrapper::GetProcesso
 			ParentTransform->setEncoding("hdr-video");
 		}
 		break;
-		case ETextureSourceEncoding::TSE_Gamma22:
+		case EEncoding::Gamma22:
 		{
 			ExponentTransformRcPtr ChildTransform = ExponentTransform::Create();
 			ChildTransform->setValue({ 2.2, 2.2, 2.2, 1.0 });
@@ -449,7 +582,7 @@ FOpenColorIOProcessorWrapper FOpenColorIOEngineBuiltInConfigWrapper::GetProcesso
 			ParentTransform->setEncoding("sdr-video");
 		}
 		break;
-		case ETextureSourceEncoding::TSE_BT1886:
+		case EEncoding::BT1886:
 		{
 			ExponentTransformRcPtr ChildTransform = ExponentTransform::Create();
 			ChildTransform->setValue({ 2.4, 2.4, 2.4, 1.0 });
@@ -459,7 +592,7 @@ FOpenColorIOProcessorWrapper FOpenColorIOEngineBuiltInConfigWrapper::GetProcesso
 			ParentTransform->setEncoding("sdr-video");
 		}
 		break;
-		case ETextureSourceEncoding::TSE_Gamma26:
+		case EEncoding::Gamma26:
 		{
 			ExponentTransformRcPtr ChildTransform = ExponentTransform::Create();
 			ChildTransform->setValue({ 2.6, 2.6, 2.6, 1.0 });
@@ -469,7 +602,7 @@ FOpenColorIOProcessorWrapper FOpenColorIOEngineBuiltInConfigWrapper::GetProcesso
 			ParentTransform->setEncoding("sdr-video");
 		}
 		break;
-		case ETextureSourceEncoding::TSE_Cineon:
+		case EEncoding::Cineon:
 		{
 			const double BlackOffset = FGenericPlatformMath::Pow(10.0, (95.0 - 685.0) / 300.0);
 			const double LinSideSlope = 1.0 - BlackOffset;
@@ -490,7 +623,7 @@ FOpenColorIOProcessorWrapper FOpenColorIOEngineBuiltInConfigWrapper::GetProcesso
 			ParentTransform->setEncoding("log");
 		}
 		break;
-		case ETextureSourceEncoding::TSE_REDLog:
+		case EEncoding::REDLog:
 		{
 			const double BlackOffset = FGenericPlatformMath::Pow(10.0, (0.0 - 1023.0) / 511.0);
 			const double LinSideSlope = 1.0 - BlackOffset;
@@ -511,7 +644,7 @@ FOpenColorIOProcessorWrapper FOpenColorIOEngineBuiltInConfigWrapper::GetProcesso
 			ParentTransform->setEncoding("log");
 		}
 		break;
-		case ETextureSourceEncoding::TSE_REDLog3G10:
+		case EEncoding::REDLog3G10:
 		{
 			static constexpr double LinSideSlope = 155.975327;
 			static constexpr double LinSideOffset = 0.01 * LinSideSlope + 1.0;
@@ -532,7 +665,7 @@ FOpenColorIOProcessorWrapper FOpenColorIOEngineBuiltInConfigWrapper::GetProcesso
 			ParentTransform->setEncoding("log");
 		}
 		break;
-		case ETextureSourceEncoding::TSE_SLog1:
+		case EEncoding::SLog1:
 		{
 			static constexpr double LinSideSlope = 1.0 / 0.9;
 			static constexpr double LinSideOffset = 0.037584;
@@ -552,7 +685,7 @@ FOpenColorIOProcessorWrapper FOpenColorIOEngineBuiltInConfigWrapper::GetProcesso
 			ParentTransform->setEncoding("log");
 		}
 		break;
-		case ETextureSourceEncoding::TSE_SLog2:
+		case EEncoding::SLog2:
 		{
 			static constexpr double LinSideSlope = 155.0 / 197.1;
 			static constexpr double LinSideOffset = 0.037584;
@@ -575,7 +708,7 @@ FOpenColorIOProcessorWrapper FOpenColorIOEngineBuiltInConfigWrapper::GetProcesso
 			ParentTransform->setEncoding("log");
 		}
 		break;
-		case ETextureSourceEncoding::TSE_SLog3:
+		case EEncoding::SLog3:
 		{
 			static constexpr double LinSideSlope = 5.26315789473684;
 			static constexpr double LinSideOffset = 0.0526315789473684;
@@ -598,7 +731,7 @@ FOpenColorIOProcessorWrapper FOpenColorIOEngineBuiltInConfigWrapper::GetProcesso
 			ParentTransform->setEncoding("log");
 		}
 		break;
-		case ETextureSourceEncoding::TSE_AlexaV3LogC:
+		case EEncoding::AlexaV3LogC:
 		{
 			static constexpr double LinSideSlope = 5.55555555555556;
 			static constexpr double LinSideOffset = 0.0522722750251688;
@@ -619,7 +752,7 @@ FOpenColorIOProcessorWrapper FOpenColorIOEngineBuiltInConfigWrapper::GetProcesso
 			ParentTransform->setEncoding("log");
 		}
 		break;
-		case ETextureSourceEncoding::TSE_CanonLog:
+		case EEncoding::CanonLog:
 		{
 			static constexpr double LinSideSlope = 10.1596;
 			static constexpr double LinSideOffset = 1.0;
@@ -640,7 +773,7 @@ FOpenColorIOProcessorWrapper FOpenColorIOEngineBuiltInConfigWrapper::GetProcesso
 			ParentTransform->setEncoding("log");
 		}
 		break;
-		case ETextureSourceEncoding::TSE_ProTune:
+		case EEncoding::ProTune:
 		{
 			static constexpr double LinSideSlope = 112.0;
 			static constexpr double LinSideOffset = 1.0f;
@@ -657,7 +790,7 @@ FOpenColorIOProcessorWrapper FOpenColorIOEngineBuiltInConfigWrapper::GetProcesso
 			ParentTransform->setEncoding("log");
 		}
 		break;
-		case ETextureSourceEncoding::TSE_VLog:
+		case EEncoding::VLog:
 		{
 			static constexpr double LinSideSlope = 1.0;
 			static constexpr double LinSideOffset = 0.00873;
@@ -687,264 +820,96 @@ FOpenColorIOProcessorWrapper FOpenColorIOEngineBuiltInConfigWrapper::GetProcesso
 			break;
 		}
 
-		const EChromaticAdaptationMethod ChromaticAdapation = static_cast<EChromaticAdaptationMethod>(InTextureColorSettings.ChromaticAdaptationMethod);
+		const EChromaticAdaptationMethod ChromaticAdapation = static_cast<EChromaticAdaptationMethod>(InColorSettings.ChromaticAdaptationMethod);
 
-		switch (InTextureColorSettings.ColorSpace)
-		{
-		case ETextureColorSpace::TCS_None:
-			break;
-		case ETextureColorSpace::TCS_Custom:
+		if (InColorSettings.ColorSpaceOverride.IsSet())
 		{
 			const FColorSpace SourceColorSpace = FColorSpace
 			(
-				InTextureColorSettings.RedChromaticityCoordinate,
-				InTextureColorSettings.GreenChromaticityCoordinate,
-				InTextureColorSettings.BlueChromaticityCoordinate,
-				InTextureColorSettings.WhiteChromaticityCoordinate
+				InColorSettings.ColorSpaceOverride.GetValue()[0],
+				InColorSettings.ColorSpaceOverride.GetValue()[1],
+				InColorSettings.ColorSpaceOverride.GetValue()[2],
+				InColorSettings.ColorSpaceOverride.GetValue()[3]
 			);
 			MatrixTransformRcPtr MatrixTransform = MatrixTransform::Create();
 			const FMatrix44d ToWorkingMat = Transpose<double>(FColorSpaceTransform(SourceColorSpace, FColorSpace::GetWorking(), ChromaticAdapation));
 			MatrixTransform->setMatrix(&ToWorkingMat.M[0][0]);
 			TransformToWCS->appendTransform(MatrixTransform);
 		}
-		break;
-		default:
+		else if (InColorSettings.ColorSpace != EColorSpace::None)
 		{
-			const FColorSpace SourceColorSpace = FColorSpace(static_cast<EColorSpace>(InTextureColorSettings.ColorSpace));
+			const FColorSpace SourceColorSpace = FColorSpace(static_cast<EColorSpace>(InColorSettings.ColorSpace));
 			const FMatrix44d ToWorkingMat = Transpose<double>(FColorSpaceTransform(SourceColorSpace, FColorSpace::GetWorking(), ChromaticAdapation));
 			MatrixTransformRcPtr MatrixTransform = MatrixTransform::Create();
 			MatrixTransform->setMatrix(&ToWorkingMat.M[0][0]);
 			TransformToWCS->appendTransform(MatrixTransform);
 		}
-		break;
-		}
 
 		ParentTransform->setTransform(TransformToWCS, TRANSFORM_DIR_FORWARD);
 
 		// Update builtin config
-		ConfigRcPtr NewConfig = Pimpl->Config->createEditableCopy();
+		ConstConfigRcPtr& CurrentConfig = BuiltInConfig.Pimpl->Config;
+		ConfigRcPtr NewConfig = CurrentConfig->createEditableCopy();
 		NewConfig->addNamedTransform(ParentTransform);
-		Pimpl->Config = NewConfig;
+		CurrentConfig = NewConfig;
 	}
 #endif // WITH_OCIO
 
-	return FOpenColorIOProcessorWrapper(this, TransformName);
-}
-
-FString FOpenColorIOEngineBuiltInConfigWrapper::GetTransformToWorkingColorSpaceName(const FTextureSourceColorSettings& InTextureColorSettings)
-{
-	const uint32 SettingsId = (uint32)InTextureColorSettings.EncodingOverride | (uint32)InTextureColorSettings.ColorSpace << 8u | (uint32)InTextureColorSettings.ChromaticAdaptationMethod << 16u;
-	FString TransformName = FString::Printf(TEXT("UE_%u"), SettingsId);
-
-	if (InTextureColorSettings.ColorSpace == ETextureColorSpace::TCS_Custom)
-	{
-		uint32 SrcChromaticityHash = 0;
-		SrcChromaticityHash ^= GetTypeHash(InTextureColorSettings.RedChromaticityCoordinate);
-		SrcChromaticityHash ^= GetTypeHash(InTextureColorSettings.GreenChromaticityCoordinate);
-		SrcChromaticityHash ^= GetTypeHash(InTextureColorSettings.BlueChromaticityCoordinate);
-		SrcChromaticityHash ^= GetTypeHash(InTextureColorSettings.WhiteChromaticityCoordinate);
-		TransformName += FString::Printf(TEXT("_%u"), SrcChromaticityHash);
-	}
-
-	return TransformName;
-}
-
-FOpenColorIOProcessorWrapper::FOpenColorIOProcessorWrapper(
-	const FOpenColorIOConfigWrapper* InConfig,
-	FStringView InSourceColorSpace,
-	FStringView InDestinationColorSpace,
-	const TMap<FString, FString>& InContextKeyValues)
-	: Pimpl(MakePimpl<FOpenColorIOProcessorPimpl, EPimplPtrMode::DeepCopy>())
-	, OwnerConfig(InConfig)
-	, WorkingColorSpaceTransformType(EOpenColorIOWorkingColorSpaceTransform::None)
-{
-	if (InSourceColorSpace == OpenColorIOWrapper::GetWorkingColorSpaceName())
-	{
-		WorkingColorSpaceTransformType = EOpenColorIOWorkingColorSpaceTransform::Source;
-	}
-	else if (InDestinationColorSpace == OpenColorIOWrapper::GetWorkingColorSpaceName())
-	{
-		WorkingColorSpaceTransformType = EOpenColorIOWorkingColorSpaceTransform::Destination;
-	}
-
-#if WITH_OCIO
-#if !PLATFORM_EXCEPTIONS_DISABLED
-	try
-#endif
-	{
-		if (OwnerConfig != nullptr && OwnerConfig->IsValid())
-		{
-			const OCIO_NAMESPACE::ConstConfigRcPtr& Config = OwnerConfig->Pimpl->Config;
-			OCIO_NAMESPACE::ContextRcPtr Context = Config->getCurrentContext()->createEditableCopy();
-
-			for (const TPair<FString, FString>& KeyValue : InContextKeyValues)
-			{
-				Context->setStringVar(TCHAR_TO_ANSI(*KeyValue.Key), TCHAR_TO_ANSI(*KeyValue.Value));
-			}
-
-			Pimpl->Processor = Config->getProcessor(
-				Context,
-				StringCast<ANSICHAR>(InSourceColorSpace.GetData()).Get(),
-				StringCast<ANSICHAR>(InDestinationColorSpace.GetData()).Get()
-			);
-		}
-	}
-#if !PLATFORM_EXCEPTIONS_DISABLED
-	catch (OCIO_NAMESPACE::Exception& Exc)
-	{
-		UE_LOG(LogOpenColorIOWrapper, Log, TEXT("Failed to create processor for [%s, %s]. Error message: %s"), InSourceColorSpace.GetData(), InDestinationColorSpace.GetData(), StringCast<TCHAR>(Exc.what()).Get());
-	}
-#endif
-#endif // WITH_OCIO
+	return FOpenColorIOWrapperProcessor(&BuiltInConfig, TransformName);
 }
 
-FOpenColorIOProcessorWrapper::FOpenColorIOProcessorWrapper(
-	const FOpenColorIOConfigWrapper* InConfig,
-	FStringView InSourceColorSpace,
-	FStringView InDisplay,
-	FStringView InView,
-	bool bInverseDirection,
-	const TMap<FString, FString>& InContextKeyValues)
-	: Pimpl(MakePimpl<FOpenColorIOProcessorPimpl, EPimplPtrMode::DeepCopy>())
-	, OwnerConfig(InConfig)
-	, WorkingColorSpaceTransformType(EOpenColorIOWorkingColorSpaceTransform::None)
+bool FOpenColorIOWrapperProcessor::IsValid() const
 {
-	if (InSourceColorSpace == OpenColorIOWrapper::GetWorkingColorSpaceName())
-	{
-		if (bInverseDirection)
-		{
-			WorkingColorSpaceTransformType = EOpenColorIOWorkingColorSpaceTransform::Destination;
-		}
-		else
-		{
-			WorkingColorSpaceTransformType = EOpenColorIOWorkingColorSpaceTransform::Source;
-		}
-	}
-
 #if WITH_OCIO
-#if !PLATFORM_EXCEPTIONS_DISABLED
-	try
-#endif
-	{
-		if (OwnerConfig != nullptr && OwnerConfig->IsValid())
-		{
-			const OCIO_NAMESPACE::ConstConfigRcPtr& Config = OwnerConfig->Pimpl->Config;
-			OCIO_NAMESPACE::ContextRcPtr Context = Config->getCurrentContext()->createEditableCopy();
-
-			for (const TPair<FString, FString>& KeyValue : InContextKeyValues)
-			{
-				Context->setStringVar(TCHAR_TO_ANSI(*KeyValue.Key), TCHAR_TO_ANSI(*KeyValue.Value));
-			}
-
-			Pimpl->Processor = Config->getProcessor(
-				StringCast<ANSICHAR>(InSourceColorSpace.GetData()).Get(),
-				StringCast<ANSICHAR>(InDisplay.GetData()).Get(),
-				StringCast<ANSICHAR>(InView.GetData()).Get(),
-				static_cast<OCIO_NAMESPACE::TransformDirection>(bInverseDirection)
-			);
-		}
-	}
-#if !PLATFORM_EXCEPTIONS_DISABLED
-	catch (OCIO_NAMESPACE::Exception& Exc)
-	{
-		UE_LOG(LogOpenColorIOWrapper, Log, TEXT("Failed to create processor for [%s, %s, %s, %s]. Error message: %s"), InSourceColorSpace.GetData(), InDisplay.GetData(), InView.GetData(), (bInverseDirection ? TEXT("Inverse") : TEXT("Forward")), StringCast<TCHAR>(Exc.what()).Get());
-	}
-#endif
+	return OwnerConfig != nullptr && Pimpl->Processor != nullptr;
+#else
+	return false;
 #endif // WITH_OCIO
 }
 
-FOpenColorIOProcessorWrapper::FOpenColorIOProcessorWrapper(const FOpenColorIOConfigWrapper* InConfig, FStringView InNamedTransform, bool bInverseDirection, const TMap<FString, FString>& InContextKeyValues)
-	: Pimpl(MakePimpl<FOpenColorIOProcessorPimpl, EPimplPtrMode::DeepCopy>())
-	, OwnerConfig(InConfig)
-	, WorkingColorSpaceTransformType(EOpenColorIOWorkingColorSpaceTransform::None)
+FString FOpenColorIOWrapperProcessor::GetTransformToWorkingColorSpaceName(const FOpenColorIOWrapperSourceColorSettings& InSourceColorSettings)
 {
-#if WITH_OCIO
-#if !PLATFORM_EXCEPTIONS_DISABLED
-	try
-#endif
-	{
-		if (OwnerConfig != nullptr && OwnerConfig->IsValid())
-		{
-			const OCIO_NAMESPACE::ConstConfigRcPtr& Config = OwnerConfig->Pimpl->Config;
-			OCIO_NAMESPACE::ContextRcPtr Context = Config->getCurrentContext()->createEditableCopy();
-
-			for (const TPair<FString, FString>& KeyValue : InContextKeyValues)
-			{
-				Context->setStringVar(TCHAR_TO_ANSI(*KeyValue.Key), TCHAR_TO_ANSI(*KeyValue.Value));
-			}
+	const uint32 SettingsId = (uint32)InSourceColorSettings.EncodingOverride | (uint32)InSourceColorSettings.ColorSpace << 8u | (uint32)InSourceColorSettings.ChromaticAdaptationMethod << 16u;
+	FString TransformName = FString::Printf(TEXT("UE_%u"), SettingsId);
 
-			Pimpl->Processor = Config->getProcessor(
-				StringCast<ANSICHAR>(InNamedTransform.GetData()).Get(),
-				static_cast<OCIO_NAMESPACE::TransformDirection>(bInverseDirection)
-			);
-		}
-	}
-#if !PLATFORM_EXCEPTIONS_DISABLED
-	catch (OCIO_NAMESPACE::Exception& Exc)
+	if (InSourceColorSettings.ColorSpaceOverride.IsSet())
 	{
-		UE_LOG(LogOpenColorIOWrapper, Log, TEXT("Failed to create processor for [%s]. Error message: %s"), InNamedTransform.GetData(), StringCast<TCHAR>(Exc.what()).Get());
+		uint32 SrcChromaticityHash = 0;
+		SrcChromaticityHash ^= GetTypeHash(InSourceColorSettings.ColorSpaceOverride.GetValue()[0]);
+		SrcChromaticityHash ^= GetTypeHash(InSourceColorSettings.ColorSpaceOverride.GetValue()[1]);
+		SrcChromaticityHash ^= GetTypeHash(InSourceColorSettings.ColorSpaceOverride.GetValue()[2]);
+		SrcChromaticityHash ^= GetTypeHash(InSourceColorSettings.ColorSpaceOverride.GetValue()[3]);
+		TransformName += FString::Printf(TEXT("_%u"), SrcChromaticityHash);
 	}
-#endif
-#endif // WITH_OCIO
-}
 
-bool FOpenColorIOProcessorWrapper::IsValid() const
-{
-#if WITH_OCIO
-	return OwnerConfig != nullptr && Pimpl->Processor != nullptr;
-#else
-	return false;
-#endif
-}
-
-FOpenColorIOCPUProcessorWrapper::FOpenColorIOCPUProcessorWrapper(FOpenColorIOProcessorWrapper InProcessor)
-	: ParentProcessor(MoveTemp(InProcessor))
-{
-}
-
-bool FOpenColorIOCPUProcessorWrapper::IsValid() const
-{
-	return ParentProcessor.IsValid();
+	return TransformName;
 }
 
-bool FOpenColorIOCPUProcessorWrapper::TransformColor(FLinearColor& InOutColor) const
+bool FOpenColorIOWrapperProcessor::TransformColor(FLinearColor& InOutColor) const
 {
 #if WITH_OCIO
-	TRACE_CPUPROFILER_EVENT_SCOPE(FOpenColorIOCPUProcessorWrapper::TransformImage)
-#if !PLATFORM_EXCEPTIONS_DISABLED
-	try
-#endif
-	{
-		using namespace OCIO_NAMESPACE;
+	TRACE_CPUPROFILER_EVENT_SCOPE(FOpenColorIOWrapperProcessor::TransformColor)
+	OCIO_EXCEPTION_HANDLING_TRY()
 
 		if (IsValid())
 		{
 			// Apply the main color transformation
-			ConstCPUProcessorRcPtr CPUProcessor = ParentProcessor.Pimpl->Processor->getOptimizedCPUProcessor(BIT_DEPTH_F32, BIT_DEPTH_F32, OPTIMIZATION_DEFAULT);
+			OCIO_NAMESPACE::ConstCPUProcessorRcPtr CPUProcessor = Pimpl->Processor->getDefaultCPUProcessor();
 			CPUProcessor->applyRGBA(&InOutColor.R);
 
 			return true;
 		}
-	}
-#if !PLATFORM_EXCEPTIONS_DISABLED
-	catch (OCIO_NAMESPACE::Exception& Exc)
-	{
-		UE_LOG(LogOpenColorIOWrapper, Log, TEXT("Failed to transform color. Error message: %s"), StringCast<TCHAR>(Exc.what()).Get());
-	}
-#endif
+	OCIO_EXCEPTION_HANDLING_CATCH(Log, TEXT("Failed to transform color. Error message: %s"))
 #endif // WITH_OCIO
 
 	return false;
 }
 
-bool FOpenColorIOCPUProcessorWrapper::TransformImage(const FImageView& InOutImage) const
+bool FOpenColorIOWrapperProcessor::TransformImage(const FImageView& InOutImage) const
 {
 #if WITH_OCIO
-	TRACE_CPUPROFILER_EVENT_SCOPE(FOpenColorIOCPUProcessorWrapper::TransformImage)
-#if !PLATFORM_EXCEPTIONS_DISABLED
-	try
-#endif
-	{
+	TRACE_CPUPROFILER_EVENT_SCOPE(FOpenColorIOWrapperProcessor::TransformImage)
+	OCIO_EXCEPTION_HANDLING_TRY()
 		using namespace OCIO_NAMESPACE;
 
 		if (IsValid())
@@ -953,13 +918,13 @@ bool FOpenColorIOCPUProcessorWrapper::TransformImage(const FImageView& InOutImag
 			if (ImageDesc)
 			{
 				const TUniquePtr<ANSICHAR[]> AnsiWorkingColorSpaceName = OpenColorIOWrapper::MakeAnsiString(OpenColorIOWrapper::GetWorkingColorSpaceName());
-				ConstConfigRcPtr	InterchangeConfig = IOpenColorIOWrapperModule::Get().GetEngineBuiltInConfig()->Pimpl->Config;
-				ConstConfigRcPtr Config = ParentProcessor.OwnerConfig->Pimpl->Config;
+				ConstConfigRcPtr	InterchangeConfig = IOpenColorIOWrapperModule::Get().GetEngineBuiltInConfig().Pimpl->Config;
+				ConstConfigRcPtr Config = OwnerConfig->Pimpl->Config;
 
 				BitDepth BitDepth = ImageDesc->getBitDepth();
 
 				// Conditionally apply a conversion from the working color space to interchange space
-				if (ParentProcessor.WorkingColorSpaceTransformType == EOpenColorIOWorkingColorSpaceTransform::Source)
+				if (WorkingColorSpaceTransformType == EOpenColorIOWorkingColorSpaceTransform::Source)
 				{
 					ConstProcessorRcPtr	InterchangeProcessor = InterchangeConfig->GetProcessorFromConfigs(
 						InterchangeConfig,
@@ -972,11 +937,11 @@ bool FOpenColorIOCPUProcessorWrapper::TransformImage(const FImageView& InOutImag
 				}
 
 				// Apply the main color transformation
-				ConstCPUProcessorRcPtr CPUProcessor = ParentProcessor.Pimpl->Processor->getOptimizedCPUProcessor(BitDepth, BitDepth, OPTIMIZATION_DEFAULT);
+				ConstCPUProcessorRcPtr CPUProcessor = Pimpl->Processor->getOptimizedCPUProcessor(BitDepth, BitDepth, OPTIMIZATION_DEFAULT);
 				CPUProcessor->apply(*ImageDesc);
 
 				// Conditionally apply a conversion from the interchange space to the working color space
-				if (ParentProcessor.WorkingColorSpaceTransformType == EOpenColorIOWorkingColorSpaceTransform::Destination)
+				if (WorkingColorSpaceTransformType == EOpenColorIOWorkingColorSpaceTransform::Destination)
 				{
 					ConstProcessorRcPtr	InterchangeProcessor = InterchangeConfig->GetProcessorFromConfigs(
 						Config,
@@ -991,26 +956,17 @@ bool FOpenColorIOCPUProcessorWrapper::TransformImage(const FImageView& InOutImag
 				return true;
 			}
 		}
-	}
-#if !PLATFORM_EXCEPTIONS_DISABLED
-	catch (OCIO_NAMESPACE::Exception& Exc)
-	{
-		UE_LOG(LogOpenColorIOWrapper, Log, TEXT("Failed to transform image. Error message: %s"), StringCast<TCHAR>(Exc.what()).Get());
-	}
-#endif
+	OCIO_EXCEPTION_HANDLING_CATCH(Log, TEXT("Failed to transform image. Error message: %s"))
 #endif // WITH_OCIO
 
 	return false;
 }
 
-bool FOpenColorIOCPUProcessorWrapper::TransformImage(const FImageView& SrcImage, const FImageView& DestImage) const
+bool FOpenColorIOWrapperProcessor::TransformImage(const FImageView& SrcImage, const FImageView& DestImage) const
 {
 #if WITH_OCIO
-	TRACE_CPUPROFILER_EVENT_SCOPE(FOpenColorIOCPUProcessorWrapper::TransformImage)
-#if !PLATFORM_EXCEPTIONS_DISABLED
-	try
-#endif
-	{
+	TRACE_CPUPROFILER_EVENT_SCOPE(FOpenColorIOWrapperProcessor::TransformImage)
+	OCIO_EXCEPTION_HANDLING_TRY()
 		using namespace OCIO_NAMESPACE;
 
 		if (IsValid())
@@ -1020,14 +976,14 @@ bool FOpenColorIOCPUProcessorWrapper::TransformImage(const FImageView& SrcImage,
 			if (SrcImageDesc && DestImageDesc)
 			{
 				const TUniquePtr<ANSICHAR[]> AnsiWorkingColorSpaceName = OpenColorIOWrapper::MakeAnsiString(OpenColorIOWrapper::GetWorkingColorSpaceName());
-				ConstConfigRcPtr	InterchangeConfig = IOpenColorIOWrapperModule::Get().GetEngineBuiltInConfig()->Pimpl->Config;
-				ConstConfigRcPtr	Config = ParentProcessor.OwnerConfig->Pimpl->Config;
+				ConstConfigRcPtr	InterchangeConfig = IOpenColorIOWrapperModule::Get().GetEngineBuiltInConfig().Pimpl->Config;
+				ConstConfigRcPtr	Config = OwnerConfig->Pimpl->Config;
 
 				BitDepth SrcBitDepth = SrcImageDesc->getBitDepth();
 				BitDepth DestBitDepth = DestImageDesc->getBitDepth();
 
 				// Conditionally apply a conversion from the working color space to interchange space
-				if (ParentProcessor.WorkingColorSpaceTransformType == EOpenColorIOWorkingColorSpaceTransform::Source)
+				if (WorkingColorSpaceTransformType == EOpenColorIOWorkingColorSpaceTransform::Source)
 				{
 					ConstProcessorRcPtr	InterchangeProcessor = InterchangeConfig->GetProcessorFromConfigs(
 						InterchangeConfig,
@@ -1041,11 +997,11 @@ bool FOpenColorIOCPUProcessorWrapper::TransformImage(const FImageView& SrcImage,
 
 
 				// Apply the main color transformation
-				ConstCPUProcessorRcPtr CPUProcessor = ParentProcessor.Pimpl->Processor->getOptimizedCPUProcessor(SrcBitDepth, DestBitDepth, OPTIMIZATION_DEFAULT);
+				ConstCPUProcessorRcPtr CPUProcessor = Pimpl->Processor->getOptimizedCPUProcessor(SrcBitDepth, DestBitDepth, OPTIMIZATION_DEFAULT);
 				CPUProcessor->apply(*SrcImageDesc, *DestImageDesc);
 
 				// Conditionally apply a conversion from the interchange space to the working color space
-				if (ParentProcessor.WorkingColorSpaceTransformType == EOpenColorIOWorkingColorSpaceTransform::Destination)
+				if (WorkingColorSpaceTransformType == EOpenColorIOWorkingColorSpaceTransform::Destination)
 				{
 					ConstProcessorRcPtr	InterchangeProcessor = InterchangeConfig->GetProcessorFromConfigs(
 						Config,
@@ -1060,13 +1016,7 @@ bool FOpenColorIOCPUProcessorWrapper::TransformImage(const FImageView& SrcImage,
 				return true;
 			}
 		}
-	}
-#if !PLATFORM_EXCEPTIONS_DISABLED
-	catch (OCIO_NAMESPACE::Exception& exception)
-	{
-		UE_LOG(LogOpenColorIOWrapper, Log, TEXT("Failed to transform image. Error message: %s"), StringCast<TCHAR>(exception.what()).Get());
-	}
-#endif
+	OCIO_EXCEPTION_HANDLING_CATCH(Log, TEXT("Failed to transform image. Error message: %s"))
 #endif // WITH_OCIO
 
 	return false;
@@ -1074,16 +1024,13 @@ bool FOpenColorIOCPUProcessorWrapper::TransformImage(const FImageView& SrcImage,
 
 
 
-FOpenColorIOGPUProcessorWrapper::FOpenColorIOGPUProcessorWrapper(FOpenColorIOProcessorWrapper InProcessor, FInitializationOptions InShaderParams)
+FOpenColorIOWrapperGPUProcessor::FOpenColorIOWrapperGPUProcessor(FOpenColorIOWrapperProcessor InProcessor, bool bUseLegacy)
 	: ParentProcessor(MoveTemp(InProcessor))
 	, GPUPimpl(MakePimpl<FOpenColorIOGPUProcessorPimpl, EPimplPtrMode::DeepCopy>())
 {
 #if WITH_OCIO
-	TRACE_CPUPROFILER_EVENT_SCOPE(FOpenColorIOGPUProcessorWrapper::FOpenColorIOGPUProcessorWrapper)	
-#if !PLATFORM_EXCEPTIONS_DISABLED
-	try
-#endif
-	{
+	TRACE_CPUPROFILER_EVENT_SCOPE(FOpenColorIOWrapperGPUProcessor::FOpenColorIOWrapperGPUProcessor)
+	OCIO_EXCEPTION_HANDLING_TRY()
 		using namespace OCIO_NAMESPACE;
 
 		if (ParentProcessor.IsValid())
@@ -1096,7 +1043,7 @@ FOpenColorIOGPUProcessorWrapper::FOpenColorIOGPUProcessorWrapper(FOpenColorIOPro
 			ConstGPUProcessorRcPtr GPUProcessor = nullptr;
 			OptimizationFlags OptFlags = FOpenColorIOProcessorPimpl::GetOptimizationFlags();
 
-			if (InShaderParams.bIsLegacy)
+			if (bUseLegacy)
 			{
 				unsigned int EdgeLength = static_cast<unsigned int>(OpenColorIOWrapper::Legacy3dEdgeLength);
 				GPUProcessor = ParentProcessor.Pimpl->Processor->getOptimizedLegacyGPUProcessor(OptFlags, EdgeLength);
@@ -1110,17 +1057,11 @@ FOpenColorIOGPUProcessorWrapper::FOpenColorIOGPUProcessorWrapper(FOpenColorIOPro
 			GPUPimpl->Processor = GPUProcessor;
 			GPUPimpl->ShaderDescription = ShaderDescription;
 		}
-	}
-#if !PLATFORM_EXCEPTIONS_DISABLED
-	catch (OCIO_NAMESPACE::Exception& Exc)
-	{
-		UE_LOG(LogOpenColorIOWrapper, Log, TEXT("Failed to fetch shader info for color transform. Error message: %s"), StringCast<TCHAR>(Exc.what()).Get());
-	}
-#endif
+	OCIO_EXCEPTION_HANDLING_CATCH(Log, TEXT("Failed to fetch shader info for color transform. Error message: %s"))
 #endif // WITH_OCIO
 }
 
-bool FOpenColorIOGPUProcessorWrapper::IsValid() const
+bool FOpenColorIOWrapperGPUProcessor::IsValid() const
 {
 #if WITH_OCIO
 	return ParentProcessor.IsValid() && GPUPimpl->Processor != nullptr && GPUPimpl->ShaderDescription != nullptr;
@@ -1129,7 +1070,7 @@ bool FOpenColorIOGPUProcessorWrapper::IsValid() const
 #endif // WITH_OCIO
 }
 
-bool FOpenColorIOGPUProcessorWrapper::GetShader(FString& OutShaderCode, FString& OutShaderCacheID) const
+bool FOpenColorIOWrapperGPUProcessor::GetShader(FString& OutShaderCode, FString& OutShaderCacheID) const
 {
 #if WITH_OCIO
 	if (IsValid())
@@ -1146,7 +1087,7 @@ bool FOpenColorIOGPUProcessorWrapper::GetShader(FString& OutShaderCode, FString&
 	return false;
 }
 
-uint32 FOpenColorIOGPUProcessorWrapper::GetNum3DTextures() const
+uint32 FOpenColorIOWrapperGPUProcessor::GetNum3DTextures() const
 {
 #if WITH_OCIO
 	if (IsValid())
@@ -1158,15 +1099,11 @@ uint32 FOpenColorIOGPUProcessorWrapper::GetNum3DTextures() const
 	return 0;
 }
 
-bool FOpenColorIOGPUProcessorWrapper::Get3DTexture(uint32 InIndex, FName& OutName, uint32& OutEdgeLength, TextureFilter& OutTextureFilter, const float*& OutData) const
+bool FOpenColorIOWrapperGPUProcessor::Get3DTexture(uint32 InIndex, FName& OutName, uint32& OutEdgeLength, TextureFilter& OutTextureFilter, const float*& OutData) const
 {
 #if WITH_OCIO
 	ensure(InIndex < GetNum3DTextures());
-
-#if !PLATFORM_EXCEPTIONS_DISABLED
-	try
-#endif
-	{
+	OCIO_EXCEPTION_HANDLING_TRY()
 		if (IsValid())
 		{
 			const ANSICHAR* TextureName = nullptr;
@@ -1190,18 +1127,12 @@ bool FOpenColorIOGPUProcessorWrapper::Get3DTexture(uint32 InIndex, FName& OutNam
 
 			return TextureName && OutEdgeLength > 0 && OutData;
 		}
-	}
-#if !PLATFORM_EXCEPTIONS_DISABLED
-	catch (OCIO_NAMESPACE::Exception& Exc)
-	{
-		UE_LOG(LogOpenColorIOWrapper, Error, TEXT("Failed to fetch 3d texture(s) info for color transform. Error message: %s"), StringCast<TCHAR>(Exc.what()).Get());
-	}
-#endif
-#endif //WITH_OCIO
+	OCIO_EXCEPTION_HANDLING_CATCH(Log, TEXT("Failed to fetch 3d texture(s) info for color transform. Error message: %s"))
+#endif // WITH_OCIO
 
 	return false;
 }
-uint32 FOpenColorIOGPUProcessorWrapper::GetNumTextures() const
+uint32 FOpenColorIOWrapperGPUProcessor::GetNumTextures() const
 {
 #if WITH_OCIO
 	if (IsValid())
@@ -1213,15 +1144,12 @@ uint32 FOpenColorIOGPUProcessorWrapper::GetNumTextures() const
 	return 0;
 }
 
-bool FOpenColorIOGPUProcessorWrapper::GetTexture(uint32 InIndex, FName& OutName, uint32& OutWidth, uint32& OutHeight, TextureFilter& OutTextureFilter, bool& bOutRedChannelOnly, const float*& OutData) const
+bool FOpenColorIOWrapperGPUProcessor::GetTexture(uint32 InIndex, FName& OutName, uint32& OutWidth, uint32& OutHeight, TextureFilter& OutTextureFilter, bool& bOutRedChannelOnly, const float*& OutData) const
 {
 	ensure(InIndex < GetNumTextures());
 
 #if WITH_OCIO
-#if !PLATFORM_EXCEPTIONS_DISABLED
-	try
-#endif
-	{
+	OCIO_EXCEPTION_HANDLING_TRY()
 		if (IsValid())
 		{
 			const ANSICHAR* TextureName = nullptr;
@@ -1242,19 +1170,13 @@ bool FOpenColorIOGPUProcessorWrapper::GetTexture(uint32 InIndex, FName& OutName,
 
 			return TextureName && OutWidth > 0 && OutHeight > 0 && OutData;
 		}
-	}
-#if !PLATFORM_EXCEPTIONS_DISABLED
-	catch (OCIO_NAMESPACE::Exception& Exc)
-	{
-		UE_LOG(LogOpenColorIOWrapper, Error, TEXT("Failed to fetch texture(s) info for color transform. Error message: %s"), StringCast<TCHAR>(Exc.what()).Get());
-	}
-#endif
+	OCIO_EXCEPTION_HANDLING_CATCH(Error, TEXT("Failed to fetch texture(s) info for color transform. Error message: %s"))
 #endif // WITH_OCIO
 
 	return false;
 }
 
-FString FOpenColorIOGPUProcessorWrapper::GetCacheID() const
+FString FOpenColorIOWrapperGPUProcessor::GetCacheID() const
 {
 #if WITH_OCIO
 	if (IsValid())
diff --git a/Engine/Plugins/Compositing/OpenColorIO/Source/OpenColorIOWrapper/Private/OpenColorIOWrapperModule.cpp b/Engine/Plugins/Compositing/OpenColorIO/Source/OpenColorIOWrapper/Private/OpenColorIOWrapperModule.cpp
--- a/Engine/Plugins/Compositing/OpenColorIO/Source/OpenColorIOWrapper/Private/OpenColorIOWrapperModule.cpp
+++ b/Engine/Plugins/Compositing/OpenColorIO/Source/OpenColorIOWrapper/Private/OpenColorIOWrapperModule.cpp
@@ -18,30 +18,36 @@ public:
 	virtual void StartupModule() override
 	{
 		FOpenColorIOLibHandler::Initialize();
-
-		EngineBuiltInConfig = MakeUnique<FOpenColorIOEngineBuiltInConfigWrapper>();
 	}
 
 	virtual void ShutdownModule() override
 	{
+		EngineBuiltInConfig.Reset();
+
 		FOpenColorIOLibHandler::Shutdown();
 	}
 	//~ End IModuleInterface interface
 
 	//~ Begin IOpenColorIOWrapperModule interface
-	virtual FOpenColorIOEngineBuiltInConfigWrapper* GetEngineBuiltInConfig() override
+	virtual FOpenColorIOWrapperEngineConfig& GetEngineBuiltInConfig() override
 	{
-		return EngineBuiltInConfig.Get();
-	}
+		FScopeLock Lock(&ConfigCriticalSection);
 
-	virtual const FOpenColorIOEngineBuiltInConfigWrapper* GetEngineBuiltInConfig() const override
-	{
-		return EngineBuiltInConfig.Get();
+		if (!EngineBuiltInConfig)
+		{
+			EngineBuiltInConfig = MakeUnique<FOpenColorIOWrapperEngineConfig>();
+		}
+
+		return *EngineBuiltInConfig;
 	}
 	//~ End IOpenColorIOWrapperModule interface
 private:
-	
-	TUniquePtr<FOpenColorIOEngineBuiltInConfigWrapper> EngineBuiltInConfig = nullptr;
+
+	// Critical section for the engine config getter.
+	FCriticalSection ConfigCriticalSection;
+
+	// Global engine config using the built-in studio config, lazily allocated.
+	TUniquePtr<FOpenColorIOWrapperEngineConfig> EngineBuiltInConfig;
 };
 
 IMPLEMENT_MODULE(FOpenColorIOWrapperModule, OpenColorIOWrapper);
diff --git a/Engine/Plugins/Compositing/OpenColorIO/Source/OpenColorIOWrapper/Private/Tests/OpenColorIOWrapperTests.cpp b/Engine/Plugins/Compositing/OpenColorIO/Source/OpenColorIOWrapper/Private/Tests/OpenColorIOWrapperTests.cpp
--- a/Engine/Plugins/Compositing/OpenColorIO/Source/OpenColorIOWrapper/Private/Tests/OpenColorIOWrapperTests.cpp
+++ b/Engine/Plugins/Compositing/OpenColorIO/Source/OpenColorIOWrapper/Private/Tests/OpenColorIOWrapperTests.cpp
@@ -7,7 +7,6 @@
 #include "Math/UnrealMathUtility.h"
 #include "Misc/AutomationTest.h"
 #include "OpenColorIOWrapper.h"
-#include "Engine/Texture.h"
 #include "TransferFunctions.h"
 
 #if WITH_DEV_AUTOMATION_TESTS
@@ -30,21 +29,18 @@ bool FOpenColorIOTransferFunctionsTest::RunTest(const FString& Parameters)
 	using namespace UE::Color;
 
 	const FLinearColor TestColor = FLinearColor(0.9f, 0.5f, 0.2f, 1.0f);
-	FOpenColorIOEngineBuiltInConfigWrapper TestConfig = {};
 
-	for (uint8 TestEncoding = static_cast<uint8>(ETextureSourceEncoding::TSE_None); TestEncoding < static_cast<uint8>(ETextureSourceEncoding::TSE_MAX); ++TestEncoding)
+	for (uint8 TestEncoding = static_cast<uint8>(EEncoding::None); TestEncoding < static_cast<uint8>(EEncoding::Max); ++TestEncoding)
 	{
 		FLinearColor Expected = UE::Color::Decode(static_cast<EEncoding>(TestEncoding), TestColor);
 
-		FTextureSourceColorSettings TestSettings = {};
-		TestSettings.EncodingOverride = static_cast<ETextureSourceEncoding>(TestEncoding);
-		TestSettings.ColorSpace = ETextureColorSpace::TCS_None;
+		FOpenColorIOWrapperSourceColorSettings TestSettings;
+		TestSettings.EncodingOverride = static_cast<EEncoding>(TestEncoding);
 
-		FOpenColorIOProcessorWrapper Processor = TestConfig.GetProcessorToWorkingColorSpace(TestSettings);
-		FOpenColorIOCPUProcessorWrapper ProcessorCPU(Processor);
+		FOpenColorIOWrapperProcessor Processor = FOpenColorIOWrapperProcessor::CreateTransformToWorkingColorSpace(TestSettings);
 
 		FLinearColor Actual = TestColor;
-		ProcessorCPU.TransformColor(Actual);
+		Processor.TransformColor(Actual);
 
 		// Note: We make the tolerance relative to the values themselves to account for larger values in PQ.
 		const float Tolerance = UE_KINDA_SMALL_NUMBER * 0.5f * (Actual.R + Expected.R);
@@ -63,25 +59,34 @@ bool FOpenColorIOTransferFunctionsTest::RunTest(const FString& Parameters)
 	TSet<FString> Keys;
 	FRandomStream RandomStream;
 	RandomStream.Initialize(42);
-	for (uint8 TestEncoding = static_cast<uint8>(ETextureSourceEncoding::TSE_None); TestEncoding < static_cast<uint8>(ETextureSourceEncoding::TSE_MAX); ++TestEncoding)
+	for (uint8 TestEncoding = static_cast<uint8>(EEncoding::None); TestEncoding < static_cast<uint8>(EEncoding::Max); ++TestEncoding)
 	{
-		for (uint8 TestColorSpace = static_cast<uint8>(ETextureColorSpace::TCS_None); TestColorSpace < static_cast<uint8>(ETextureColorSpace::TCS_MAX); ++TestColorSpace)
+		for (uint8 TestColorSpace = static_cast<uint8>(EColorSpace::None); TestColorSpace <= static_cast<uint8>(EColorSpace::Max); ++TestColorSpace)
 		{
 			for (uint8 TestChromatic = 0; TestChromatic < 2; ++TestChromatic)
 			{
-				FTextureSourceColorSettings Settings = {};
-				Settings.EncodingOverride = static_cast<ETextureSourceEncoding>(TestEncoding);
-				Settings.ColorSpace = static_cast<ETextureColorSpace>(TestColorSpace);
-				if (Settings.ColorSpace == ETextureColorSpace::TCS_Custom)
+				FOpenColorIOWrapperSourceColorSettings TestSettings;
+				TestSettings.EncodingOverride = static_cast<EEncoding>(TestEncoding);
+
+				if (TestColorSpace < static_cast<uint8>(EColorSpace::Max))
+				{
+					TestSettings.ColorSpace = static_cast<EColorSpace>(TestColorSpace);
+				}
+				else
 				{
-					Settings.RedChromaticityCoordinate = FVector2D(RandomStream.GetUnitVector());
-					Settings.GreenChromaticityCoordinate = FVector2D(RandomStream.GetUnitVector());
-					Settings.BlueChromaticityCoordinate = FVector2D(RandomStream.GetUnitVector());
-					Settings.WhiteChromaticityCoordinate = FVector2D(RandomStream.GetUnitVector());
+					// We locally use EColorSpace::Max as a test custom color space.
+					TStaticArray<FVector2d, 4> RandomChromaticities;
+					RandomChromaticities[0] = FVector2d(RandomStream.GetUnitVector());
+					RandomChromaticities[1] = FVector2d(RandomStream.GetUnitVector());
+					RandomChromaticities[2] = FVector2d(RandomStream.GetUnitVector());
+					RandomChromaticities[3] = FVector2d(RandomStream.GetUnitVector());
+
+					TestSettings.ColorSpaceOverride = MoveTemp(RandomChromaticities);
 				}
-				Settings.ChromaticAdaptationMethod = static_cast<ETextureChromaticAdaptationMethod>(TestChromatic);
+				
+				TestSettings.ChromaticAdaptationMethod = static_cast<EChromaticAdaptationMethod>(TestChromatic);
 
-				Keys.Add(FOpenColorIOEngineBuiltInConfigWrapper::GetTransformToWorkingColorSpaceName(Settings));
+				Keys.Add(FOpenColorIOWrapperProcessor::GetTransformToWorkingColorSpaceName(TestSettings));
 				Count++;
 			}
 		}
diff --git a/Engine/Plugins/Compositing/OpenColorIO/Source/OpenColorIOWrapper/Public/OpenColorIOWrapper.h b/Engine/Plugins/Compositing/OpenColorIO/Source/OpenColorIOWrapper/Public/OpenColorIOWrapper.h
--- a/Engine/Plugins/Compositing/OpenColorIO/Source/OpenColorIOWrapper/Public/OpenColorIOWrapper.h
+++ b/Engine/Plugins/Compositing/OpenColorIO/Source/OpenColorIOWrapper/Public/OpenColorIOWrapper.h
@@ -4,26 +4,17 @@
 
 #include "CoreTypes.h"
 #include "Containers/ContainersFwd.h"
-#include "Containers/StringView.h"
-#include "Containers/UnrealString.h"
-#include "Containers/Map.h"
+#include "Containers/StaticArray.h"
+#include "Math/Vector.h"
 #include "Templates/PimplPtr.h"
 #include "Templates/UniquePtr.h"
 
+#include "ColorManagementDefines.h"
 #include "OpenColorIOWrapperDefines.h"
 
 struct FImageView;
-struct FTextureSourceColorSettings;
 enum TextureFilter : int;
 
-namespace UE {
-	namespace Color {
-		enum class EEncoding: uint8;
-		enum class EColorSpace : uint8;
-		class FColorSpace;
-	}
-}
-
 namespace OpenColorIOWrapper
 {
 	/** Color space name of the engine's working color space we insert in OpenColorIO configs. */
@@ -45,20 +36,24 @@ namespace OpenColorIOWrapper
 	OPENCOLORIOWRAPPER_API const TCHAR* GetVersion();
 }
 
-class OPENCOLORIOWRAPPER_API FOpenColorIOConfigWrapper
+
+/**
+ * Wrapper around library-native config object.
+ */
+class FOpenColorIOWrapperConfig
 {
 public:
 
 	/** Config initialization options. */
-	struct FInitializationOptions
+	enum EAddWorkingColorSpaceOption
 	{
-		FInitializationOptions() : bAddWorkingColorSpace(false) { };
-
-		bool bAddWorkingColorSpace;
+		WCS_None,
+		// Add the working color space as ACES 2065-1 so that we can target it with an extra conversion from our config if it defines aces_interchange.
+		WCS_AsInterchangeSpace
 	};
 
 	/** Constructor. */
-	FOpenColorIOConfigWrapper();
+	FOpenColorIOWrapperConfig();
 
 	/**
 	* Constructor.
@@ -66,80 +61,90 @@ public:
 	* @param InFilePath Config absolute file path.
 	* @param InOptions Initialization options.
 	*/
-	explicit FOpenColorIOConfigWrapper(FStringView InFilePath, FInitializationOptions InOptions);
+	OPENCOLORIOWRAPPER_API explicit FOpenColorIOWrapperConfig(FStringView InFilePath, EAddWorkingColorSpaceOption InOption = WCS_None);
 
-	virtual ~FOpenColorIOConfigWrapper() = default;
+	OPENCOLORIOWRAPPER_API virtual ~FOpenColorIOWrapperConfig() = default;
 
 	/** Valid when the native config has been successfully created and isn't null. */
-	bool IsValid() const;
+	OPENCOLORIOWRAPPER_API bool IsValid() const;
 	
 	/** Get the number of color spaces in the configuration. */
-	int32 GetNumColorSpaces() const;
+	OPENCOLORIOWRAPPER_API int32 GetNumColorSpaces() const;
 	
 	/** Get a color space name at an index. */
-	FString GetColorSpaceName(int32 Index) const;
+	OPENCOLORIOWRAPPER_API FString GetColorSpaceName(int32 Index) const;
 	
 	/** Get the index of a color space, -1 if none. */
-	int32 GetColorSpaceIndex(const TCHAR* InColorSpaceName);
+	OPENCOLORIOWRAPPER_API int32 GetColorSpaceIndex(const TCHAR* InColorSpaceName);
 	
 	/** Get the family name for a color space. */
-	FString GetColorSpaceFamilyName(const TCHAR* InColorSpaceName) const;
+	OPENCOLORIOWRAPPER_API FString GetColorSpaceFamilyName(const TCHAR* InColorSpaceName) const;
 	
 	/** Get the number of displays in the configuration. */
-	int32 GetNumDisplays() const;
+	OPENCOLORIOWRAPPER_API int32 GetNumDisplays() const;
 	
 	/** Get a display name at an index. */
-	FString GetDisplayName(int32 Index) const;
+	OPENCOLORIOWRAPPER_API FString GetDisplayName(int32 Index) const;
 	
 	/** Get the number of views for a display. */
-	int32 GetNumViews(const TCHAR* InDisplayName) const;
+	OPENCOLORIOWRAPPER_API int32 GetNumViews(const TCHAR* InDisplayName) const;
 	
 	/** Get a view name for its display and index. */
-	FString GetViewName(const TCHAR* InDisplayName, int32 Index) const;
+	OPENCOLORIOWRAPPER_API FString GetViewName(const TCHAR* InDisplayName, int32 Index) const;
 
 	/** Get a display-view trnasform name. */
-	FString GetDisplayViewTransformName(const TCHAR* InDisplayName, const TCHAR* InViewName) const;
+	OPENCOLORIOWRAPPER_API FString GetDisplayViewTransformName(const TCHAR* InDisplayName, const TCHAR* InViewName) const;
 
 	/** Get the string hash of the config. */
-	FString GetCacheID() const;
+	OPENCOLORIOWRAPPER_API FString GetCacheID() const;
 
 	/** Returns a config debug string. */
-	FString GetDebugString() const;
+	OPENCOLORIOWRAPPER_API FString GetDebugString() const;
 
 protected:
 
 	TPimplPtr<struct FOpenColorIOConfigPimpl, EPimplPtrMode::DeepCopy> Pimpl;
 
-	friend class FOpenColorIOProcessorWrapper;
-	friend class FOpenColorIOCPUProcessorWrapper;
-	friend class FOpenColorIOGPUProcessorWrapper;
+	friend class FOpenColorIOWrapperProcessor;
+	friend class FOpenColorIOWrapperGPUProcessor;
 };
 
-class OPENCOLORIOWRAPPER_API FOpenColorIOEngineBuiltInConfigWrapper : public FOpenColorIOConfigWrapper
+/* Matches FTextureSourceColorSettings in Engine/Texture.h */
+struct FOpenColorIOWrapperSourceColorSettings
+{
+	/** Source encoding of the texture, exposing more options than just sRGB. */
+	UE::Color::EEncoding EncodingOverride = UE::Color::EEncoding::None;
+
+	/** Source color space of the texture. Leave to None for custom color space. */
+	UE::Color::EColorSpace ColorSpace = UE::Color::EColorSpace::None;
+
+	/** Optional chromaticity coordinates of a custom source color space. */
+	TOptional<TStaticArray<FVector2d, 4>> ColorSpaceOverride = {};
+
+	/** Chromatic adaption method applied if the source white point differs from the working color space white point. */
+	UE::Color::EChromaticAdaptationMethod ChromaticAdaptationMethod = UE::Color::EChromaticAdaptationMethod::None;
+};
+
+/**
+ * Wrapper around library-native config object for module-specific use cases.
+ * This uses the built-in studio config as a base.
+ */
+class FOpenColorIOWrapperEngineConfig : public FOpenColorIOWrapperConfig
 {
 public:
 	/**
 	* Constructor.
 	*/
-	FOpenColorIOEngineBuiltInConfigWrapper();
-	FOpenColorIOEngineBuiltInConfigWrapper(FStringView InFilePath, FInitializationOptions InOptions) = delete;
-
-	/**
-	* Get the processor to working color space from source texture settings.
-	* @param InTextureColorSettings Source encoding and colorspace.
-	* @return Processor of the transform
-	*/
-	class FOpenColorIOProcessorWrapper GetProcessorToWorkingColorSpace(const FTextureSourceColorSettings& InTextureColorSettings);
+	FOpenColorIOWrapperEngineConfig();
 
-	/**
-	* Get the generated transform name from texture settings to working color space.
-	* @param InTextureColorSettings Source encoding and colorspace.
-	* @return Generated name of the transform
-	*/
-	static FString GetTransformToWorkingColorSpaceName(const FTextureSourceColorSettings& InTextureColorSettings);
+	FOpenColorIOWrapperEngineConfig(FStringView InFilePath, EAddWorkingColorSpaceOption InOption = WCS_None) = delete;
 };
 
-class OPENCOLORIOWRAPPER_API FOpenColorIOProcessorWrapper final
+/**
+ * Wrapper around library-native processor object.
+ * This class can also be used for CPU transforms. It will allocate CPU processors as needed, and reuse them since they are by default cached by the library.
+ */
+class FOpenColorIOWrapperProcessor final
 {
 public:
 	/**
@@ -150,11 +155,12 @@ public:
 	* @param InDestinationColorSpace Destination color space name.
 	* @param InContextKeyValues (Optional) Additional context modifiers.
 	*/
-	explicit FOpenColorIOProcessorWrapper(
-		const FOpenColorIOConfigWrapper* InConfig,
+	OPENCOLORIOWRAPPER_API explicit FOpenColorIOWrapperProcessor(
+		const FOpenColorIOWrapperConfig* InConfig,
 		FStringView InSourceColorSpace,
 		FStringView InDestinationColorSpace,
-		const TMap<FString, FString>& InContextKeyValues = {});
+		const TMap<FString, FString>& InContextKeyValues = {}
+	);
 
 	/**
 	* Constructor.
@@ -166,13 +172,14 @@ public:
 	* @param bInverseDirection Flag for inverse transform direction.
 	* @param InContextKeyValues (Optional) Additional context modifiers.
 	*/
-	explicit FOpenColorIOProcessorWrapper(
-		const FOpenColorIOConfigWrapper* InConfig,
+	OPENCOLORIOWRAPPER_API explicit FOpenColorIOWrapperProcessor(
+		const FOpenColorIOWrapperConfig* InConfig,
 		FStringView InSourceColorSpace,
 		FStringView InDisplay,
 		FStringView InView,
 		bool bInverseDirection = false,
-		const TMap<FString, FString>& InContextKeyValues = {});
+		const TMap<FString, FString>& InContextKeyValues = {}
+	);
 
 	/**
 	* Constructor.
@@ -182,65 +189,58 @@ public:
 	* @param bInverseDirection Flag for inverse transform direction.
 	* @param InContextKeyValues (Optional) Additional context modifiers.
 	*/
-	explicit FOpenColorIOProcessorWrapper(
-		const FOpenColorIOConfigWrapper* InConfig,
+	OPENCOLORIOWRAPPER_API explicit FOpenColorIOWrapperProcessor(
+		const FOpenColorIOWrapperConfig* InConfig,
 		FStringView InNamedTransform,
 		bool bInverseDirection = false,
-		const TMap<FString, FString>& InContextKeyValues = {});
-
-	/** Valid when the processor has been successfully created and isn't null. */
-	bool IsValid() const;
-
-private:
-
-	TPimplPtr<struct FOpenColorIOProcessorPimpl, EPimplPtrMode::DeepCopy> Pimpl;
-
-	const FOpenColorIOConfigWrapper* OwnerConfig;
-
-	EOpenColorIOWorkingColorSpaceTransform WorkingColorSpaceTransformType;
-
-	friend class FOpenColorIOCPUProcessorWrapper;
-	friend class FOpenColorIOGPUProcessorWrapper;
-};
-
+		const TMap<FString, FString>& InContextKeyValues = {}
+	);
 
-class OPENCOLORIOWRAPPER_API FOpenColorIOCPUProcessorWrapper final
-{
-public:
 	/**
-	* Constructor.
-	*
-	* @param InProcessor Parent processor.
+	* Create a processor from source color settings to working color space.
+	* @param InColorSettings Source encoding and colorspace.
+	* @return Processor of the transform
 	*/
-	explicit FOpenColorIOCPUProcessorWrapper(FOpenColorIOProcessorWrapper InProcessor);
+	OPENCOLORIOWRAPPER_API static FOpenColorIOWrapperProcessor CreateTransformToWorkingColorSpace(
+		const FOpenColorIOWrapperSourceColorSettings& InSourceColorSettings
+	);
 
 	/** Valid when the processor has been successfully created and isn't null. */
-	bool IsValid() const;
+	OPENCOLORIOWRAPPER_API bool IsValid() const;
 
-	/** Apply the color transform in-place to the specified color. */
-	bool TransformColor(FLinearColor& InOutColor) const;
+	/**
+	* Get the generated transform name from source color settings to working color space.
+	* @param InColorSettings Source encoding and colorspace.
+	* @return Generated name of the transform
+	*/
+	OPENCOLORIOWRAPPER_API static FString GetTransformToWorkingColorSpaceName(const FOpenColorIOWrapperSourceColorSettings& InSourceColorSettings);
 
-	/** Apply the color transform in-place to the specified image. */
-	bool TransformImage(const FImageView& InOutImage) const;
+	/** Apply the CPU color transform in-place to the specified color. */
+	OPENCOLORIOWRAPPER_API bool TransformColor(FLinearColor& InOutColor) const;
 
-	/** Apply the color transform from the source image to the destination image. (The destination FImageView is const but what it points at is not.) */
-	bool TransformImage(const FImageView& SrcImage, const FImageView& DestImage) const;
+	/** Apply the CPU color transform in-place to the specified image. */
+	OPENCOLORIOWRAPPER_API bool TransformImage(const FImageView& InOutImage) const;
+
+	/** Apply the CPU color transform from the source image to the destination image. (The destination FImageView is const but what it points at is not.) */
+	OPENCOLORIOWRAPPER_API bool TransformImage(const FImageView& SrcImage, const FImageView& DestImage) const;
 
 private:
 
-	const FOpenColorIOProcessorWrapper ParentProcessor;
+	TPimplPtr<struct FOpenColorIOProcessorPimpl, EPimplPtrMode::DeepCopy> Pimpl;
+
+	const FOpenColorIOWrapperConfig* OwnerConfig;
+
+	EOpenColorIOWorkingColorSpaceTransform WorkingColorSpaceTransformType;
+
+	friend class FOpenColorIOWrapperGPUProcessor;
 };
 
-class OPENCOLORIOWRAPPER_API FOpenColorIOGPUProcessorWrapper final
+/**
+ * Wrapper around library-native GPU processor object.
+ */
+class FOpenColorIOWrapperGPUProcessor final
 {
 public:
-	/** Gpu processor initialization options. */
-	struct FInitializationOptions
-	{
-		FInitializationOptions() : bIsLegacy(false) { }
-
-		bool bIsLegacy;
-	};
 
 	/**
 	* Constructor.
@@ -248,13 +248,13 @@ public:
 	* @param InProcessor Parent processor.
 	* @param InOptions Initialization options.
 	*/
-	explicit FOpenColorIOGPUProcessorWrapper(
-		FOpenColorIOProcessorWrapper InProcessor,
-		FInitializationOptions InOptions = FInitializationOptions()
+	OPENCOLORIOWRAPPER_API explicit FOpenColorIOWrapperGPUProcessor(
+		FOpenColorIOWrapperProcessor InProcessor,
+		bool bUseLegacy = false
 	);
 
 	/** Valid when the processor has been successfully created and isn't null. */
-	bool IsValid() const;
+	OPENCOLORIOWRAPPER_API bool IsValid() const;
 
 	/**
 	* Get the generated shader.
@@ -263,10 +263,10 @@ public:
 	* @param OutShaderCacheID Generated shader hash string.
 	* @return True when the result is valid.
 	*/
-	bool GetShader(FString& OutShaderCode, FString& OutShaderCacheID) const;
+	OPENCOLORIOWRAPPER_API bool GetShader(FString& OutShaderCode, FString& OutShaderCacheID) const;
 
 	/** Get the number of 3D LUT textures. */
-	uint32 GetNum3DTextures() const;
+	OPENCOLORIOWRAPPER_API uint32 GetNum3DTextures() const;
 
 	/**
 	* Get the index-specific 3D LUT texture used by the shader transform.
@@ -278,10 +278,10 @@ public:
 	* @param OutData Raw texel data.
 	* @return True when the result is valid.
 	*/
-	bool Get3DTexture(uint32 InIndex, FName& OutName, uint32& OutEdgeLength, TextureFilter& OutTextureFilter, const float*& OutData) const;
+	OPENCOLORIOWRAPPER_API bool Get3DTexture(uint32 InIndex, FName& OutName, uint32& OutEdgeLength, TextureFilter& OutTextureFilter, const float*& OutData) const;
 
 	/** Get the number of 2D LUT textures. (1D resources are disabled in our case.) */
-	uint32 GetNumTextures() const;
+	OPENCOLORIOWRAPPER_API uint32 GetNumTextures() const;
 
 	/**
 	* Get the index-specific 2D LUT texture used by the shader transform.
@@ -295,14 +295,14 @@ public:
 	* @param OutData Raw texel data.
 	* @return True when the result is valid.
 	*/
-	bool GetTexture(uint32 InIndex, FName& OutName, uint32& OutWidth, uint32& OutHeight, TextureFilter& OutTextureFilter, bool& bOutRedChannelOnly, const float*& OutData) const;
+	OPENCOLORIOWRAPPER_API bool GetTexture(uint32 InIndex, FName& OutName, uint32& OutWidth, uint32& OutHeight, TextureFilter& OutTextureFilter, bool& bOutRedChannelOnly, const float*& OutData) const;
 
 	/** Get the string hash of the gpu processor. */
-	FString GetCacheID() const;
+	OPENCOLORIOWRAPPER_API FString GetCacheID() const;
 
 private:
 
-	const FOpenColorIOProcessorWrapper ParentProcessor;
+	const FOpenColorIOWrapperProcessor ParentProcessor;
 
 	TPimplPtr<struct FOpenColorIOGPUProcessorPimpl, EPimplPtrMode::DeepCopy> GPUPimpl;
 };
diff --git a/Engine/Plugins/Compositing/OpenColorIO/Source/OpenColorIOWrapper/Public/OpenColorIOWrapperModule.h b/Engine/Plugins/Compositing/OpenColorIO/Source/OpenColorIOWrapper/Public/OpenColorIOWrapperModule.h
--- a/Engine/Plugins/Compositing/OpenColorIO/Source/OpenColorIOWrapper/Public/OpenColorIOWrapperModule.h
+++ b/Engine/Plugins/Compositing/OpenColorIO/Source/OpenColorIOWrapper/Public/OpenColorIOWrapperModule.h
@@ -5,7 +5,7 @@
 #include "Modules/ModuleManager.h"
 #include "Logging/LogMacros.h"
 
-class FOpenColorIOEngineBuiltInConfigWrapper;
+class FOpenColorIOWrapperEngineConfig;
 
 DECLARE_LOG_CATEGORY_EXTERN(LogOpenColorIOWrapper, Log, All);
 
@@ -41,8 +41,7 @@ public:
 	/**
 	 * Returns a minimal dynamically-created native config for engine (working color spaces) conversions.
 	 */
-	virtual FOpenColorIOEngineBuiltInConfigWrapper* GetEngineBuiltInConfig() = 0;
-	virtual const FOpenColorIOEngineBuiltInConfigWrapper* GetEngineBuiltInConfig() const = 0;
+	virtual FOpenColorIOWrapperEngineConfig& GetEngineBuiltInConfig() = 0;
 
 	/** Virtual destructor */
 	virtual ~IOpenColorIOWrapperModule() = default;
