OCIO: Interface changes for serialization refactor.

#rnx

diff --git a/Engine/Plugins/Compositing/OpenColorIO/Source/OpenColorIOWrapper/Private/OpenColorIOWrapper.cpp b/Engine/Plugins/Compositing/OpenColorIO/Source/OpenColorIOWrapper/Private/OpenColorIOWrapper.cpp
--- a/Engine/Plugins/Compositing/OpenColorIO/Source/OpenColorIOWrapper/Private/OpenColorIOWrapper.cpp
+++ b/Engine/Plugins/Compositing/OpenColorIO/Source/OpenColorIOWrapper/Private/OpenColorIOWrapper.cpp
@@ -59,6 +59,13 @@ const TCHAR* OpenColorIOWrapper::GetVersion()
 #endif // WITH_OCIO
 }
 
+void OpenColorIOWrapper::ClearAllCaches()
+{
+#if WITH_OCIO
+	OCIO_NAMESPACE::ClearAllCaches();
+#endif
+}
+
 struct FOpenColorIOConfigPimpl
 {
 #if WITH_OCIO
@@ -867,6 +874,18 @@ bool FOpenColorIOWrapperProcessor::IsValid() const
 #endif // WITH_OCIO
 }
 
+FString FOpenColorIOWrapperProcessor::GetCacheID() const
+{
+#if WITH_OCIO
+	if (IsValid())
+	{
+		return StringCast<TCHAR>(Pimpl->Processor->getCacheID()).Get();
+	}
+#endif // WITH_OCIO
+
+	return {};
+}
+
 FString FOpenColorIOWrapperProcessor::GetTransformToWorkingColorSpaceName(const FOpenColorIOWrapperSourceColorSettings& InSourceColorSettings)
 {
 	const uint32 SettingsId = (uint32)InSourceColorSettings.EncodingOverride | (uint32)InSourceColorSettings.ColorSpace << 8u | (uint32)InSourceColorSettings.ChromaticAdaptationMethod << 16u;
@@ -1070,15 +1089,15 @@ bool FOpenColorIOWrapperGPUProcessor::IsValid() const
 #endif // WITH_OCIO
 }
 
-bool FOpenColorIOWrapperGPUProcessor::GetShader(FString& OutShaderCode, FString& OutShaderCacheID) const
+bool FOpenColorIOWrapperGPUProcessor::GetShader(FString& OutShaderCacheID, FString& OutShaderCode) const
 {
 #if WITH_OCIO
 	if (IsValid())
 	{
 		ensureMsgf(GPUPimpl->ShaderDescription->getNumDynamicProperties() == 0, TEXT("We do not currently support dynamic properties."));
 
-		OutShaderCode = StringCast<TCHAR>(GPUPimpl->ShaderDescription->getShaderText()).Get();
 		OutShaderCacheID = StringCast<TCHAR>(GPUPimpl->ShaderDescription->getCacheID()).Get();
+		OutShaderCode = StringCast<TCHAR>(GPUPimpl->ShaderDescription->getShaderText()).Get();
 
 		return true;
 	}
diff --git a/Engine/Plugins/Compositing/OpenColorIO/Source/OpenColorIOWrapper/Public/OpenColorIOWrapper.h b/Engine/Plugins/Compositing/OpenColorIO/Source/OpenColorIOWrapper/Public/OpenColorIOWrapper.h
--- a/Engine/Plugins/Compositing/OpenColorIO/Source/OpenColorIOWrapper/Public/OpenColorIOWrapper.h
+++ b/Engine/Plugins/Compositing/OpenColorIO/Source/OpenColorIOWrapper/Public/OpenColorIOWrapper.h
@@ -34,6 +34,9 @@ namespace OpenColorIOWrapper
 
 	/** Get the OpenColorIO version string. */
 	OPENCOLORIOWRAPPER_API const TCHAR* GetVersion();
+
+	/** Calls the native function of the same name. */
+	OPENCOLORIOWRAPPER_API void ClearAllCaches();
 }
 
 
@@ -208,6 +211,9 @@ public:
 	/** Valid when the processor has been successfully created and isn't null. */
 	OPENCOLORIOWRAPPER_API bool IsValid() const;
 
+	/** Get the string hash of the processor. */
+	OPENCOLORIOWRAPPER_API FString GetCacheID() const;
+
 	/**
 	* Get the generated transform name from source color settings to working color space.
 	* @param InColorSettings Source encoding and colorspace.
@@ -259,11 +265,11 @@ public:
 	/**
 	* Get the generated shader.
 	*
-	* @param OutShaderCode Generated shader text.
 	* @param OutShaderCacheID Generated shader hash string.
+	* @param OutShaderCode Generated shader text.
 	* @return True when the result is valid.
 	*/
-	OPENCOLORIOWRAPPER_API bool GetShader(FString& OutShaderCode, FString& OutShaderCacheID) const;
+	OPENCOLORIOWRAPPER_API bool GetShader(FString& OutShaderCacheID, FString& OutShaderCode) const;
 
 	/** Get the number of 3D LUT textures. */
 	OPENCOLORIOWRAPPER_API uint32 GetNum3DTextures() const;
