OCIO: Cache generated transform textures and shader code into uasset.

* Further interface cleanup: deprecating now unused functions, wrapper editor-only functions explicitely, renamed TransformImage.

diff --git a/Engine/Plugins/Compositing/OpenColorIO/Source/OpenColorIO/Private/OpenColorIOColorTransform.cpp b/Engine/Plugins/Compositing/OpenColorIO/Source/OpenColorIO/Private/OpenColorIOColorTransform.cpp
--- a/Engine/Plugins/Compositing/OpenColorIO/Source/OpenColorIO/Private/OpenColorIOColorTransform.cpp
+++ b/Engine/Plugins/Compositing/OpenColorIO/Source/OpenColorIO/Private/OpenColorIOColorTransform.cpp
@@ -9,6 +9,7 @@
 #include "OpenColorIOConfiguration.h"
 #include "OpenColorIOModule.h"
 #include "OpenColorIOSettings.h"
+#include "UObject/UE5MainStreamObjectVersion.h"
 #include "UObject/UObjectIterator.h"
 #include "DataDrivenShaderPlatformInfo.h"
 
@@ -126,7 +127,7 @@
 			ERHIFeatureLevel::Type LoadedFeatureLevel = LoadedShaderMap->GetShaderMapId().FeatureLevel;
 			if (!OutColorTransformResourcesLoaded[LoadedFeatureLevel])
 			{
-				OutColorTransformResourcesLoaded[LoadedFeatureLevel] = Owner->AllocateResource();
+				OutColorTransformResourcesLoaded[LoadedFeatureLevel] = new FOpenColorIOTransformResource();
 			}
 
 			OutColorTransformResourcesLoaded[LoadedFeatureLevel]->SetInlineShaderMap(LoadedShaderMap);
@@ -162,19 +163,59 @@
 
 bool UOpenColorIOColorTransform::Initialize(UOpenColorIOConfiguration* InOwner, const FString& InSourceColorSpace, const FString& InDestinationColorSpace, const TMap<FString, FString>& InContextKeyValues)
 {
+#if WITH_EDITOR
 	return Initialize(InSourceColorSpace, InDestinationColorSpace, InContextKeyValues);
+#else
+	return false;
+#endif
 }
 
 bool UOpenColorIOColorTransform::Initialize(UOpenColorIOConfiguration* InOwner, const FString& InSourceColorSpace, const FString& InDisplay, const FString& InView, EOpenColorIOViewTransformDirection InDirection, const TMap<FString, FString>& InContextKeyValues)
 {
+#if WITH_EDITOR
 	return Initialize(InSourceColorSpace, InDisplay, InView, InDirection, InContextKeyValues);
+#else
+	return false;
+#endif
 }
 
+void UOpenColorIOColorTransform::Serialize(FArchive& Ar)
+{
+	Super::Serialize(Ar);
+
+	Ar.UsingCustomVersion(FUE5MainStreamObjectVersion::GUID);
+
+#if WITH_EDITOR
+	SerializeOpenColorIOShaderMaps(&CachedColorTransformResourcesForCooking, Ar, LoadedTransformResources);
+#else
+	SerializeOpenColorIOShaderMaps(nullptr, Ar, LoadedTransformResources);
+#endif
+	
+	if (Ar.CustomVer(FUE5MainStreamObjectVersion::GUID) < FUE5MainStreamObjectVersion::OpenColorIOAssetCacheSerialization)
+	{
+		if (Ar.IsLoading())
+		{
+			// The asset previously always saved the number of LUTs, and conditionally the textures when cooked.
+			// We can safely ignore the cooked case since the current logic will be applied. However, we must
+			// still pretend to read the number of textures for the archive load to match the expected size.
+
+			int32 NumTexturesDeprecated = 0;
+			Ar << NumTexturesDeprecated;
+		}
+	}
+}
+
+
+#if WITH_EDITOR
+
 bool UOpenColorIOColorTransform::Initialize(const FString& InSourceColorSpace, const FString& InDestinationColorSpace, const TMap<FString, FString>& InContextKeyValues)
 {
-	ContextKeyValues = InContextKeyValues;
 
-#if WITH_EDITOR
+	if (InSourceColorSpace.IsEmpty() || InDestinationColorSpace.IsEmpty())
+	{
+		return false;
+	}
+
 	if (InSourceColorSpace == OpenColorIOWrapper::GetWorkingColorSpaceName())
 	{
 		WorkingColorSpaceTransformType = EOpenColorIOWorkingColorSpaceTransform::Source;
@@ -183,16 +224,25 @@
 	{
 		WorkingColorSpaceTransformType = EOpenColorIOWorkingColorSpaceTransform::Destination;
 	}
-#endif
 
-	return GenerateColorTransformData(InSourceColorSpace, InDestinationColorSpace);
+	SourceColorSpace = InSourceColorSpace;
+	DestinationColorSpace = InDestinationColorSpace;
+	bIsDisplayViewType = false;
+	ContextKeyValues = InContextKeyValues;
+
+	ProcessTransform();
+	CacheResourceShadersForRendering(true);
+
+	return true;
 }
 
 bool UOpenColorIOColorTransform::Initialize(const FString& InSourceColorSpace, const FString& InDisplay, const FString& InView, EOpenColorIOViewTransformDirection InDirection, const TMap<FString, FString>& InContextKeyValues)
 {
-	ContextKeyValues = InContextKeyValues;
+	if (InSourceColorSpace.IsEmpty() || InDisplay.IsEmpty() || InView.IsEmpty())
+	{
+		return false;
+	}
 
-#if WITH_EDITOR
 	if (InSourceColorSpace == OpenColorIOWrapper::GetWorkingColorSpaceName())
 	{
 		if (InDirection == EOpenColorIOViewTransformDirection::Forward)
@@ -204,35 +254,28 @@
 			WorkingColorSpaceTransformType = EOpenColorIOWorkingColorSpaceTransform::Destination;
 		}
 	}
-#endif
 
-	return GenerateColorTransformData(InSourceColorSpace, InDisplay, InView, InDirection);
+	SourceColorSpace = InSourceColorSpace;
+	DestinationColorSpace.Empty();
+	Display = InDisplay;
+	View = InView;
+	bIsDisplayViewType = true;
+	DisplayViewDirection = InDirection;
+	ContextKeyValues = InContextKeyValues;
+
+	ProcessTransform();
+	CacheResourceShadersForRendering(true);
+
+	return true;
 }
 
-void UOpenColorIOColorTransform::Serialize(FArchive& Ar)
-{
-	Super::Serialize(Ar);
-
-#if WITH_EDITOR
-	SerializeOpenColorIOShaderMaps(&CachedColorTransformResourcesForCooking, Ar, LoadedTransformResources);
-#else
-	SerializeOpenColorIOShaderMaps(nullptr, Ar, LoadedTransformResources);
-#endif
-
-	SerializeLuts(Ar);
-}
-
-
 void UOpenColorIOColorTransform::CacheResourceShadersForCooking(EShaderPlatform InShaderPlatform, const ITargetPlatform* TargetPlatform, const FString& InShaderHash, const FString& InShaderCode, const FString& InRawConfigHash, TArray<FOpenColorIOTransformResource*>& OutCachedResources)
 {
 	const ERHIFeatureLevel::Type TargetFeatureLevel = GetMaxSupportedFeatureLevel(InShaderPlatform);
 
-	FOpenColorIOTransformResource* NewResource = AllocateResource();
-#if WITH_EDITOR
+	FOpenColorIOTransformResource* NewResource = new FOpenColorIOTransformResource();
+
 	FName AssetPath = GetOutermost()->GetFName();
-#else
-	FName AssetPath;
-#endif
 	NewResource->SetupResource((ERHIFeatureLevel::Type)TargetFeatureLevel, InShaderHash, InShaderCode, InRawConfigHash, GetTransformFriendlyName(), AssetPath, WorkingColorSpaceTransformType);
 
 	const bool bApplyCompletedShaderMap = false;
@@ -242,103 +285,190 @@
 	OutCachedResources.Add(NewResource);
 }
 
-void UOpenColorIOColorTransform::SerializeLuts(FArchive& Ar)
+
+void UOpenColorIOColorTransform::ProcessTransform()
 {
+	const UOpenColorIOConfiguration* ConfigurationOwner = Cast<UOpenColorIOConfiguration>(GetOuter());
+	const FOpenColorIOWrapperConfig* ConfigWrapper = GetTransformConfigWrapper(ConfigurationOwner);
 
-	if (Ar.IsSaving())
+	if (ConfigWrapper != nullptr)
 	{
-		int32 Num3dLutsToSave = 0;
-		if (Ar.IsCooking())
+		const FOpenColorIOWrapperProcessor TransformProcessor = GetTransformProcessor(this, ConfigWrapper);
+
+		if (TransformProcessor.IsValid())
 		{
-			Num3dLutsToSave = Textures.Num();
-		}
+			Textures.Reset();
 
-		Ar << Num3dLutsToSave;
+			const UOpenColorIOSettings* Settings = GetDefault<UOpenColorIOSettings>();
+			const FOpenColorIOWrapperGPUProcessor GPUProcessor = FOpenColorIOWrapperGPUProcessor(TransformProcessor, Settings->bUseLegacyProcessor);
+			const FString GpuProcessorHash = GPUProcessor.GetCacheID();
 
-		if (Num3dLutsToSave > 0)
-		{
-			Ar << Textures;
+			// Process 3D luts
+			for (uint32 Index = 0; Index < GPUProcessor.GetNum3DTextures(); ++Index)
+			{
+				TextureFilter Filter;
+				FName TextureName;
+				uint32 EdgeLength = 0;
+				const float* TextureValues = 0x0;
+
+				bool bSuccess = GPUProcessor.Get3DTexture(Index, TextureName, EdgeLength, Filter, TextureValues);
+				check(bSuccess);
+
+				TObjectPtr<UTexture> Result = CreateTexture3DLUT(GpuProcessorHash, TextureName, EdgeLength, Filter, TextureValues);
+
+				const int32 SlotIndex = TextureName.GetNumber() - 1; // Rely on FName's index number extraction for convenience
+				Textures.Add(SlotIndex, MoveTemp(Result));
+			}
+
+			// Process 1D luts
+			for (uint32 Index = 0; Index < GPUProcessor.GetNumTextures(); ++Index)
+			{
+				FName TextureName;
+				uint32 TextureWidth = 0;
+				uint32 TextureHeight = 0;
+				TextureFilter Filter;
+				bool bRedChannelOnly = false;
+				const float* TextureValues = 0x0;
+
+				bool bSuccess = GPUProcessor.GetTexture(Index, TextureName, TextureWidth, TextureHeight, Filter, bRedChannelOnly, TextureValues);
+				checkf(bSuccess, TEXT("Failed to read OCIO 1D LUT data."));
+
+				TObjectPtr<UTexture> Result = CreateTexture1DLUT(GpuProcessorHash, TextureName, TextureWidth, TextureHeight, Filter, bRedChannelOnly, TextureValues);
+
+				const int32 SlotIndex = TextureName.GetNumber() - 1; // Rely on FName's index number extraction for convenience
+				Textures.Add(SlotIndex, MoveTemp(Result));
+			}
+
+			ensureAlwaysMsgf(Textures.Num() <= (int32)OpenColorIOShader::MaximumTextureSlots, TEXT("Color transform %s exceeds our current limit of %u texture slots. Use the legacy processor instead."), *GetTransformFriendlyName(), OpenColorIOShader::MaximumTextureSlots);
+			
+			// Generate shader code
+			GPUProcessor.GetShader(GeneratedShaderHash, GeneratedShader);
+
+			// Early return on success
+			return;
 		}
 	}
-	else if (Ar.IsLoading())
+
+	UE_LOG(LogOpenColorIO, Error, TEXT("Failed to process and cache resource(s) for color transform %s. Configuration file [%s] was invalid."), *GetTransformFriendlyName(), ConfigurationOwner ? *ConfigurationOwner->ConfigurationFile.FilePath : TEXT("Missing"));
+}
+
+TObjectPtr<UTexture> UOpenColorIOColorTransform::CreateTexture3DLUT(const FString& InProcessorIdentifier, const FName& InName, uint32 InLutLength, TextureFilter InFilter, const float* InSourceData)
+{
+	check(InSourceData);
+	const UOpenColorIOSettings* Settings = GetDefault<UOpenColorIOSettings>();
+
+	/* Note here that while it is possible to create proper 32f textures using [UTexture]::CreateTransient and reparenting
+	 * via Texture->Rename(nullptr, this), cooking would fail as it remains unsupported currently for those formats.
+	 * (The same note applies to 1D LUT creation.) */
+
+	TObjectPtr<UVolumeTexture> OutTexture = NewObject<UVolumeTexture>(this, InName);
+
+	//Initializes source data with the raw LUT. If it's found in DDC, the resulting platform data will be fetched from there. 
+	//If not, the source data will be used to generate the platform data.
+	OutTexture->MipGenSettings = TMGS_NoMipmaps;
+	OutTexture->SRGB = 0;
+	OutTexture->LODGroup = TEXTUREGROUP_ColorLookupTable;
+	if (Settings->bUse32fLUT)
 	{
-		int32 NumLoaded3dLuts = 0;
-		Ar << NumLoaded3dLuts;
+		// 32f resources have to be explicitely requested using this compression setting.
+		OutTexture->CompressionSettings = TextureCompressionSettings::TC_HDR_F32;
+	}
+	else
+	{
+		OutTexture->CompressionNone = true;
+	}
+	OutTexture->Filter = InFilter;
+	OutTexture->AddressMode = TextureAddress::TA_Clamp;
+	OutTexture->Source.Init(InLutLength, InLutLength, InLutLength, /*NumMips=*/ 1, TSF_RGBA32F, nullptr);
 
-		if (NumLoaded3dLuts > 0)
+	FLinearColor* MipData = reinterpret_cast<FLinearColor*>(OutTexture->Source.LockMip(0));
+	for (uint32 Z = 0; Z < InLutLength; ++Z)
+	{
+		for (uint32 Y = 0; Y < InLutLength; Y++)
 		{
-			//Will only happen on cooked data
-			Ar << Textures;
+			FLinearColor* Row = &MipData[Y * InLutLength + Z * InLutLength * InLutLength];
+			const float* Source = &InSourceData[Y * InLutLength * 3 + Z * InLutLength * InLutLength * 3];
+			for (uint32 X = 0; X < InLutLength; X++)
+			{
+				FLinearColor& TargetColor = Row[X];
+				TargetColor.R = Source[X * 3 + 0];
+				TargetColor.G = Source[X * 3 + 1];
+				TargetColor.B = Source[X * 3 + 2];
+				TargetColor.A = 1.0f;
+			}
 		}
 	}
+	OutTexture->Source.UnlockMip(0);
+
+	//Generate a Guid from the identifier received from the library and our DDC version.
+	FGuid LutGuid;
+	GetOpenColorIOLUTKeyGuid(InProcessorIdentifier, InName, LutGuid);
+	OutTexture->Source.SetId(LutGuid, true);
+
+	//Process our new texture to be usable in rendering pipeline.
+	OutTexture->UpdateResource();
+
+	return OutTexture;
 }
 
-void UOpenColorIOColorTransform::CacheResourceTextures()
+TObjectPtr<UTexture> UOpenColorIOColorTransform::CreateTexture1DLUT(const FString& InProcessorIdentifier, const FName& InName, uint32 InTextureWidth, uint32 InTextureHeight, TextureFilter InFilter, bool bRedChannelOnly, const float* InSourceData)
 {
-#if WITH_EDITOR
-	if (Textures.IsEmpty())
+	check(InSourceData);
+	const UOpenColorIOSettings* Settings = GetDefault<UOpenColorIOSettings>();
+
+	TObjectPtr<UTexture2D> OutTexture = NewObject<UTexture2D>(this, InName);
+	OutTexture->MipGenSettings = TMGS_NoMipmaps;
+	OutTexture->SRGB = 0;
+	OutTexture->LODGroup = TEXTUREGROUP_ColorLookupTable;
+	if (Settings->bUse32fLUT)
 	{
-		const UOpenColorIOSettings* Settings = GetDefault<UOpenColorIOSettings>();
-		const UOpenColorIOConfiguration* ConfigurationOwner = Cast<UOpenColorIOConfiguration>(GetOuter());
-		const FOpenColorIOWrapperConfig* ConfigWrapper = GetTransformConfigWrapper(ConfigurationOwner);
+		// 32f resources have to be explicitely requested using this compression setting.
+		OutTexture->CompressionSettings = bRedChannelOnly ? TextureCompressionSettings::TC_SingleFloat : TextureCompressionSettings::TC_HDR_F32;
+	}
+	else
+	{
+		OutTexture->CompressionNone = true;
+	}
+	OutTexture->Filter = InFilter;
+	OutTexture->AddressX = TextureAddress::TA_Clamp;
+	OutTexture->AddressY = TextureAddress::TA_Clamp;
 
-		if (ConfigWrapper != nullptr)
+	if (bRedChannelOnly)
+	{
+		OutTexture->Source.Init(InTextureWidth, InTextureHeight, /*NumSlices=*/ 1, /*NumMips=*/ 1, TSF_R32F, reinterpret_cast<const uint8*>(InSourceData));
+	}
+	else
+	{
+		OutTexture->Source.Init(InTextureWidth, InTextureHeight, /*NumSlices=*/ 1, /*NumMips=*/ 1, TSF_RGBA32F, nullptr);
+
+		FLinearColor* MipData = reinterpret_cast<FLinearColor*>(OutTexture->Source.LockMip(0));
+		for (uint32 Y = 0; Y < InTextureHeight; Y++)
 		{
-			const FOpenColorIOWrapperProcessor TransformProcessor = GetTransformProcessor(this, ConfigWrapper);
-			if (TransformProcessor.IsValid())
+			FLinearColor* Row = &MipData[Y * InTextureWidth];
+			const float* Source = &InSourceData[Y * InTextureWidth * 3];
+			for (uint32 X = 0; X < InTextureWidth; X++)
 			{
-				const FOpenColorIOWrapperGPUProcessor GPUProcessor = FOpenColorIOWrapperGPUProcessor(TransformProcessor, Settings->bUseLegacyProcessor);
-				const FString ProcessorID = GPUProcessor.GetCacheID();
-
-				//In editor, it will use what's on DDC if there's something corresponding to the actual data or use that raw data
-				//that OCIO library has on board. The textures will be serialized only when cooking.
-
-				// Process 3D luts
-				for (uint32 Index = 0; Index < GPUProcessor.GetNum3DTextures(); ++Index)
-				{
-					TextureFilter Filter;
-					FName TextureName;
-					uint32 EdgeLength = 0;
-					const float* TextureValues = 0x0;
-
-					bool bSuccess = GPUProcessor.Get3DTexture(Index, TextureName, EdgeLength, Filter, TextureValues);
-					check(bSuccess);
-
-					TObjectPtr<UTexture> Result = CreateTexture3DLUT(ProcessorID, TextureName, EdgeLength, Filter, TextureValues);
-
-					const int32 SlotIndex = TextureName.GetNumber() - 1; // Rely on FName's index number extraction for convenience
-					Textures.Add(SlotIndex, MoveTemp(Result));
-				}
-
-				// Process 1D luts
-				for (uint32 Index = 0; Index < GPUProcessor.GetNumTextures(); ++Index)
-				{
-					FName TextureName;
-					uint32 TextureWidth = 0;
-					uint32 TextureHeight = 0;
-					TextureFilter Filter;
-					bool bRedChannelOnly = false;
-					const float* TextureValues = 0x0;
-
-					bool bSuccess = GPUProcessor.GetTexture(Index, TextureName, TextureWidth, TextureHeight, Filter, bRedChannelOnly, TextureValues);
-					checkf(bSuccess, TEXT("Failed to read OCIO 1D LUT data."));
-
-					TObjectPtr<UTexture> Result = CreateTexture1DLUT(ProcessorID, TextureName, TextureWidth, TextureHeight, Filter, bRedChannelOnly, TextureValues);
-
-					const int32 SlotIndex = TextureName.GetNumber() - 1; // Rely on FName's index number extraction for convenience
-					Textures.Add(SlotIndex, MoveTemp(Result));
-				}
-
-				ensureAlwaysMsgf(Textures.Num() <= (int32)OpenColorIOShader::MaximumTextureSlots, TEXT("Color transform %s exceeds our current limit of %u texture slots. Use the legacy processor instead."), *GetTransformFriendlyName(), OpenColorIOShader::MaximumTextureSlots);
+				FLinearColor& TargetColor = Row[X];
+				TargetColor.R = Source[X * 3 + 0];
+				TargetColor.G = Source[X * 3 + 1];
+				TargetColor.B = Source[X * 3 + 2];
+				TargetColor.A = 1.0f;
 			}
 		}
-		else
-		{
-			UE_LOG(LogOpenColorIO, Error, TEXT("Failed to cache texture resource(s) for color transform %s. Configuration file [%s] was invalid."), *GetTransformFriendlyName(), ConfigurationOwner ? *ConfigurationOwner->ConfigurationFile.FilePath : TEXT("Missing"));
-		}
+		OutTexture->Source.UnlockMip(0);
 	}
+
+	//Generate a Guid from the identifier received from the library and our DDC version.
+	FGuid LutGuid;
+	GetOpenColorIOLUTKeyGuid(InProcessorIdentifier, InName, LutGuid);
+	OutTexture->Source.SetId(LutGuid, true);
+
+	//Process our new texture to be usable in rendering pipeline.
+	OutTexture->UpdateResource();
+
+	return OutTexture;
+}
 #endif // WITH_EDITOR
-}
 
 void UOpenColorIOColorTransform::CacheResourceShadersForRendering(bool bRegenerateId)
 {
@@ -349,11 +479,8 @@
 
 	if (FApp::CanEverRender())
 	{
-		//Update shader hash to fetch pre-compiled shader from DDC and grab shader code to be able to compile it on the fly if it's missing
-		FString ShaderCodeHash;
-		FString ShaderCode;
-		FString RawConfigHash;
-		if (UpdateShaderInfo(ShaderCodeHash, ShaderCode, RawConfigHash))
+		const UOpenColorIOConfiguration* ConfigurationOwner = Cast<UOpenColorIOConfiguration>(GetOuter());
+		if (IsValid(ConfigurationOwner))
 		{
 			//OCIO shaders are simple, we should be compatible with any feature levels. Use the levels required for materials.
 			uint32 FeatureLevelsToCompile = UMaterialInterface::GetFeatureLevelsToCompileForAllMaterials();
@@ -365,14 +492,19 @@
 				FOpenColorIOTransformResource*& TransformResource = ColorTransformResources[CacheFeatureLevel];
 				if (TransformResource == nullptr)
 				{
-					TransformResource = AllocateResource();
+					TransformResource = new FOpenColorIOTransformResource();
 				}
 
-#if WITH_EDITOR
-				FName AssetPath = GetOutermost()->GetFName();
-#else
+				FString ShaderCodeHash;
+				FString ShaderCode;
+				FString RawConfigHash;
 				FName AssetPath;
-#endif
+#if WITH_EDITOR
+				ShaderCodeHash = GeneratedShaderHash;
+				ShaderCode = GeneratedShader;
+				RawConfigHash = ConfigurationOwner->GetConfigWrapper()->GetCacheID();
+				AssetPath = GetOutermost()->GetFName();
+#endif // WITH_EDITOR
 				TransformResource->SetupResource(CacheFeatureLevel, ShaderCodeHash, ShaderCode, RawConfigHash, GetTransformFriendlyName(), AssetPath, WorkingColorSpaceTransformType);
 
 				const bool bApplyCompletedShaderMap = true;
@@ -494,7 +626,7 @@
 }
 
 #if WITH_EDITOR
-bool UOpenColorIOColorTransform::EditorTransformImage(const FImageView& InOutImage) const
+bool UOpenColorIOColorTransform::TransformImage(const FImageView& InOutImage) const
 {
 	const UOpenColorIOConfiguration* ConfigurationOwner = Cast<UOpenColorIOConfiguration>(GetOuter());
 	const FOpenColorIOWrapperConfig* ConfigWrapper = GetTransformConfigWrapper(ConfigurationOwner);
@@ -515,7 +647,7 @@
 	return false;
 }
 
-bool UOpenColorIOColorTransform::EditorTransformImage(const FImageView& SrcImage, const FImageView& DestImage) const
+bool UOpenColorIOColorTransform::TransformImage(const FImageView& SrcImage, const FImageView& DestImage) const
 {
 	const UOpenColorIOConfiguration* ConfigurationOwner = Cast<UOpenColorIOConfiguration>(GetOuter());
 	const FOpenColorIOWrapperConfig* ConfigWrapper = GetTransformConfigWrapper(ConfigurationOwner);
@@ -557,49 +689,6 @@
 	}
 }
 
-bool UOpenColorIOColorTransform::GenerateColorTransformData(const FString& InSourceColorSpace, const FString& InDestinationColorSpace)
-{
-#if WITH_EDITOR
-	if (InSourceColorSpace.IsEmpty() || InDestinationColorSpace.IsEmpty())
-	{
-		return false;
-	}
-
-	SourceColorSpace = InSourceColorSpace;
-	DestinationColorSpace = InDestinationColorSpace;
-	bIsDisplayViewType = false;
-
-	CacheResourceTextures();
-	CacheResourceShadersForRendering(true);
-
-	return true;
-#endif //WITH_EDITOR
-	return false;
-}
-
-bool UOpenColorIOColorTransform::GenerateColorTransformData(const FString& InSourceColorSpace, const FString& InDisplay, const FString& InView, EOpenColorIOViewTransformDirection InDirection)
-{
-#if WITH_EDITOR
-	if (InSourceColorSpace.IsEmpty() || InDisplay.IsEmpty() || InView.IsEmpty())
-	{
-		return false;
-	}
-
-	SourceColorSpace = InSourceColorSpace;
-	DestinationColorSpace.Empty();
-	Display = InDisplay;
-	View = InView;
-	bIsDisplayViewType = true;
-	DisplayViewDirection = InDirection;
-
-	CacheResourceTextures();
-	CacheResourceShadersForRendering(true);
-
-	return true;
-#endif //WITH_EDITOR
-	return false;
-}
-
 FString UOpenColorIOColorTransform::GetTransformFriendlyName() const
 {
 	if (bIsDisplayViewType)
@@ -638,7 +727,7 @@
 			{
 				OutRawConfigHash = ConfigurationOwner->GetConfigWrapper()->GetCacheID();
 
-				return GPUProcessor.GetShader(OutShaderCode, OutShaderCodeHash);
+				return GPUProcessor.GetShader(OutShaderCodeHash, OutShaderCode);
 			}
 		}
 	}
@@ -653,131 +742,6 @@
 #endif //WITH_EDITOR
 }
 
-TObjectPtr<UTexture> UOpenColorIOColorTransform::CreateTexture3DLUT(const FString& InProcessorIdentifier, const FName& InName, uint32 InLutLength, TextureFilter InFilter, const float* InSourceData)
-{
-	TObjectPtr<UVolumeTexture> OutTexture = nullptr;
-
-#if WITH_EDITOR
-	check(InSourceData);
-	const UOpenColorIOSettings* Settings = GetDefault<UOpenColorIOSettings>();
-
-	/* Note here that while it is possible to create proper 32f textures using [UTexture]::CreateTransient and reparenting
-	 * via Texture->Rename(nullptr, this), cooking would fail as it remains unsupported currently for those formats.
-	 * (The same note applies to 1D LUT creation.) */
-
-	OutTexture = NewObject<UVolumeTexture>(this, InName);
-
-	//Initializes source data with the raw LUT. If it's found in DDC, the resulting platform data will be fetched from there. 
-	//If not, the source data will be used to generate the platform data.
-	OutTexture->MipGenSettings = TMGS_NoMipmaps;
-	OutTexture->SRGB = 0;
-	OutTexture->LODGroup = TEXTUREGROUP_ColorLookupTable;
-	if (Settings->bUse32fLUT)
-	{
-		// 32f resources have to be explicitely requested using this compression setting.
-		OutTexture->CompressionSettings = TextureCompressionSettings::TC_HDR_F32;
-	}
-	else
-	{
-		OutTexture->CompressionNone = true;
-	}
-	OutTexture->Filter = InFilter;
-	OutTexture->AddressMode = TextureAddress::TA_Clamp;
-	OutTexture->Source.Init(InLutLength, InLutLength, InLutLength, /*NumMips=*/ 1, TSF_RGBA32F, nullptr);
-
-	FLinearColor* MipData = reinterpret_cast<FLinearColor*>(OutTexture->Source.LockMip(0));
-	for (uint32 Z = 0; Z < InLutLength; ++Z)
-	{
-		for (uint32 Y = 0; Y < InLutLength; Y++)
-		{
-			FLinearColor* Row = &MipData[Y * InLutLength + Z * InLutLength * InLutLength];
-			const float* Source = &InSourceData[Y * InLutLength * 3 + Z * InLutLength * InLutLength * 3];
-			for (uint32 X = 0; X < InLutLength; X++)
-			{
-				FLinearColor& TargetColor = Row[X];
-				TargetColor.R = Source[X * 3 + 0];
-				TargetColor.G = Source[X * 3 + 1];
-				TargetColor.B = Source[X * 3 + 2];
-				TargetColor.A = 1.0f;
-			}
-		}
-	}
-	OutTexture->Source.UnlockMip(0);
-
-	//Generate a Guid from the identifier received from the library and our DDC version.
-	FGuid LutGuid;
-	GetOpenColorIOLUTKeyGuid(InProcessorIdentifier, InName, LutGuid);
-	OutTexture->Source.SetId(LutGuid, true);
-
-	//Process our new texture to be usable in rendering pipeline.
-	OutTexture->UpdateResource();
-#endif
-
-	return OutTexture;
-}
-
-TObjectPtr<UTexture> UOpenColorIOColorTransform::CreateTexture1DLUT(const FString& InProcessorIdentifier, const FName& InName, uint32 InTextureWidth, uint32 InTextureHeight, TextureFilter InFilter, bool bRedChannelOnly, const float* InSourceData)
-{
-	TObjectPtr<UTexture2D> OutTexture = nullptr;
-
-#if WITH_EDITOR
-	check(InSourceData);
-	const UOpenColorIOSettings* Settings = GetDefault<UOpenColorIOSettings>();
-
-	OutTexture = NewObject<UTexture2D>(this, InName);
-	OutTexture->MipGenSettings = TMGS_NoMipmaps;
-	OutTexture->SRGB = 0;
-	OutTexture->LODGroup = TEXTUREGROUP_ColorLookupTable;
-	if (Settings->bUse32fLUT)
-	{
-		// 32f resources have to be explicitely requested using this compression setting.
-		OutTexture->CompressionSettings = bRedChannelOnly ? TextureCompressionSettings::TC_SingleFloat : TextureCompressionSettings::TC_HDR_F32;
-	}
-	else
-	{
-		OutTexture->CompressionNone = true;
-	}
-	OutTexture->Filter = InFilter;
-	OutTexture->AddressX = TextureAddress::TA_Clamp;
-	OutTexture->AddressY = TextureAddress::TA_Clamp;
-
-	if (bRedChannelOnly)
-	{
-		OutTexture->Source.Init(InTextureWidth, InTextureHeight, /*NumSlices=*/ 1, /*NumMips=*/ 1, TSF_R32F, reinterpret_cast<const uint8*>(InSourceData));
-	}
-	else
-	{
-		OutTexture->Source.Init(InTextureWidth, InTextureHeight, /*NumSlices=*/ 1, /*NumMips=*/ 1, TSF_RGBA32F, nullptr);
-
-		FLinearColor* MipData = reinterpret_cast<FLinearColor*>(OutTexture->Source.LockMip(0));
-		for (uint32 Y = 0; Y < InTextureHeight; Y++)
-		{
-			FLinearColor* Row = &MipData[Y * InTextureWidth];
-			const float* Source = &InSourceData[Y * InTextureWidth * 3];
-			for (uint32 X = 0; X < InTextureWidth; X++)
-			{
-				FLinearColor& TargetColor = Row[X];
-				TargetColor.R = Source[X * 3 + 0];
-				TargetColor.G = Source[X * 3 + 1];
-				TargetColor.B = Source[X * 3 + 2];
-				TargetColor.A = 1.0f;
-			}
-		}
-		OutTexture->Source.UnlockMip(0);
-	}
-
-	//Generate a Guid from the identifier received from the library and our DDC version.
-	FGuid LutGuid;
-	GetOpenColorIOLUTKeyGuid(InProcessorIdentifier, InName, LutGuid);
-	OutTexture->Source.SetId(LutGuid, true);
-
-	//Process our new texture to be usable in rendering pipeline.
-	OutTexture->UpdateResource();
-#endif
-
-	return OutTexture;
-}
-
 void UOpenColorIOColorTransform::FlushResourceShaderMaps()
 {
 	if (FApp::CanEverRender())
@@ -830,7 +794,14 @@
 		
 		if (bIsTransformActive)
 		{
-			CacheResourceTextures();
+#if WITH_EDITOR
+			// Recently saved transforms will have a valid generated shader which we can use as an indication to skip processing on load.
+			if (GeneratedShader.IsEmpty())
+			{
+				ProcessTransform();
+			}
+#endif
+
 			CacheResourceShadersForRendering(false);
 		}
 		else
@@ -889,18 +860,16 @@
 
 	if (DesiredShaderFormats.Num() > 0)
 	{
-		//Need to re-update shader data when cooking. They won't have been previously fetched.
-		FString ShaderCodeHash;
-		FString ShaderCode;
-		FString RawConfigHash;
-		if (UpdateShaderInfo(ShaderCodeHash, ShaderCode, RawConfigHash))
+		const UOpenColorIOConfiguration* ConfigurationOwner = Cast<UOpenColorIOConfiguration>(GetOuter());
+		if (ConfigurationOwner != nullptr)
 		{
+			FString RawConfigHash = ConfigurationOwner->GetConfigWrapper()->GetCacheID();
 			// Cache for all the shader formats that the cooking target requires
 			for (int32 FormatIndex = 0; FormatIndex < DesiredShaderFormats.Num(); FormatIndex++)
 			{
 				const EShaderPlatform LegacyShaderPlatform = ShaderFormatToLegacyShaderPlatform(DesiredShaderFormats[FormatIndex]);
 				// Begin caching shaders for the target platform and store the FOpenColorIOTransformResource being compiled into CachedColorTransformResourcesForCooking
-				CacheResourceShadersForCooking(LegacyShaderPlatform, TargetPlatform, ShaderCodeHash, ShaderCode, RawConfigHash, *CachedColorTransformResourceForPlatformPtr);
+				CacheResourceShadersForCooking(LegacyShaderPlatform, TargetPlatform, GeneratedShaderHash, GeneratedShader, RawConfigHash, *CachedColorTransformResourceForPlatformPtr);
 			}
 		}
 	}
diff --git a/Engine/Plugins/Compositing/OpenColorIO/Source/OpenColorIO/Private/OpenColorIOConfiguration.cpp b/Engine/Plugins/Compositing/OpenColorIO/Source/OpenColorIO/Private/OpenColorIOConfiguration.cpp
--- a/Engine/Plugins/Compositing/OpenColorIO/Source/OpenColorIO/Private/OpenColorIOConfiguration.cpp
+++ b/Engine/Plugins/Compositing/OpenColorIO/Source/OpenColorIO/Private/OpenColorIOConfiguration.cpp
@@ -134,21 +134,21 @@ bool UOpenColorIOConfiguration::Validate() const
 }
 
 #if WITH_EDITOR
-bool UOpenColorIOConfiguration::EditorTransformImage(const FOpenColorIOColorConversionSettings& InSettings, const FImageView& InOutImage) const
+bool UOpenColorIOConfiguration::TransformImage(const FOpenColorIOColorConversionSettings& InSettings, const FImageView& InOutImage) const
 {
 	if (const TObjectPtr<UOpenColorIOColorTransform>* TransformPtr = FindTransform(InSettings))
 	{
-		return (*TransformPtr)->EditorTransformImage(InOutImage);
+		return (*TransformPtr)->TransformImage(InOutImage);
 	}
 
 	return false;
 }
 
-bool UOpenColorIOConfiguration::EditorTransformImage(const FOpenColorIOColorConversionSettings& InSettings, const FImageView& SrcImage, const FImageView& DestImage) const
+bool UOpenColorIOConfiguration::TransformImage(const FOpenColorIOColorConversionSettings& InSettings, const FImageView& SrcImage, const FImageView& DestImage) const
 {
 	if (const TObjectPtr<UOpenColorIOColorTransform>* TransformPtr = FindTransform(InSettings))
 	{
-		return (*TransformPtr)->EditorTransformImage(SrcImage, DestImage);
+		return (*TransformPtr)->TransformImage(SrcImage, DestImage);
 	}
 
 	return false;
@@ -158,72 +158,103 @@ bool UOpenColorIOConfiguration::EditorTransformImage(const FOpenColorIOColorConv
 void UOpenColorIOConfiguration::ReloadExistingColorspaces()
 {
 #if WITH_EDITOR
-	TArray<FOpenColorIOColorSpace> ColorSpacesToBeReloaded = DesiredColorSpaces;
-	TArray<FOpenColorIODisplayView> DisplayViewsToBeReloaded = DesiredDisplayViews;
-	DesiredColorSpaces.Reset();
-	DesiredDisplayViews.Reset();
-	ColorTransforms.Reset();
 	LoadConfiguration();
 
-	if (Config == nullptr)
+	if (Config && Config->IsValid())
 	{
-		return;
-	}
+		const FString LoadedConfigHash = Config->GetCacheID() + FString(OpenColorIOWrapper::GetVersion());
 
-	// This will make sure that all colorspaces are up to date in case an index, family or name is changed.
-	for (const FOpenColorIOColorSpace& ExistingColorSpace : ColorSpacesToBeReloaded)
-	{
-		int ColorSpaceIndex = Config->GetColorSpaceIndex(*ExistingColorSpace.ColorSpaceName);
-		if (ColorSpaceIndex < 0)
+		// Hash is different, proceed with the regeneration...
+		if (ConfigHash != LoadedConfigHash)
 		{
-			// Name not found, therefore we don't need to re-add this colorspace.
-			continue;
-		}
+			ConfigHash = LoadedConfigHash;
 
-		FOpenColorIOColorSpace ColorSpace;
-		ColorSpace.ColorSpaceIndex = ColorSpaceIndex;
-		ColorSpace.ColorSpaceName = ExistingColorSpace.ColorSpaceName;
-		ColorSpace.FamilyName = Config->GetColorSpaceFamilyName(*ExistingColorSpace.ColorSpaceName);
-		DesiredColorSpaces.Add(ColorSpace);
-	}
+			TArray<FOpenColorIOColorSpace> ColorSpacesToBeReloaded = DesiredColorSpaces;
+			TArray<FOpenColorIODisplayView> DisplayViewsToBeReloaded = DesiredDisplayViews;
+			DesiredColorSpaces.Reset();
+			DesiredDisplayViews.Reset();
+			ColorTransforms.Reset();
 
-	for (const FOpenColorIODisplayView& ExistingDisplayView : DisplayViewsToBeReloaded)
-	{
-		for (int32 ViewIndex = 0; ViewIndex < Config->GetNumViews(*ExistingDisplayView.Display); ++ViewIndex)
-		{
-			const FString ViewName = Config->GetViewName(*ExistingDisplayView.Display, ViewIndex);
-			
-			if (ViewName == ExistingDisplayView.View)
+			// This will make sure that all colorspaces are up to date in case an index, family or name is changed.
+			for (const FOpenColorIOColorSpace& ExistingColorSpace : ColorSpacesToBeReloaded)
 			{
-				DesiredDisplayViews.Add(ExistingDisplayView);
-			}
-		}
-	}
+				int ColorSpaceIndex = Config->GetColorSpaceIndex(*ExistingColorSpace.ColorSpaceName);
+				if (ColorSpaceIndex < 0)
+				{
+					// Name not found, therefore we don't need to re-add this colorspace.
+					UE_LOG(LogOpenColorIO, Display, TEXT("Removing %s, not found."), *ExistingColorSpace.ToString());
+					continue;
+				}
 
-	const UOpenColorIOSettings* Settings = GetDefault<UOpenColorIOSettings>();
+				FOpenColorIOColorSpace ColorSpace;
+				ColorSpace.ColorSpaceIndex = ColorSpaceIndex;
+				ColorSpace.ColorSpaceName = ExistingColorSpace.ColorSpaceName;
+				ColorSpace.FamilyName = Config->GetColorSpaceFamilyName(*ExistingColorSpace.ColorSpaceName);
+				DesiredColorSpaces.Add(ColorSpace);
+			}
 
-	// Genereate new shaders.
-	for (int32 indexTop = 0; indexTop < DesiredColorSpaces.Num(); ++indexTop)
-	{
-		const FOpenColorIOColorSpace& TopColorSpace = DesiredColorSpaces[indexTop];
+			for (const FOpenColorIODisplayView& ExistingDisplayView : DisplayViewsToBeReloaded)
+			{
+				bool bDisplayViewFound = false;
+				for (int32 ViewIndex = 0; ViewIndex < Config->GetNumViews(*ExistingDisplayView.Display); ++ViewIndex)
+				{
+					const FString ViewName = Config->GetViewName(*ExistingDisplayView.Display, ViewIndex);
 
-		for (int32 indexOther = indexTop + 1; indexOther < DesiredColorSpaces.Num(); ++indexOther)
-		{
-			const FOpenColorIOColorSpace& OtherColorSpace = DesiredColorSpaces[indexOther];
+					if (ViewName == ExistingDisplayView.View)
+					{
+						DesiredDisplayViews.Add(ExistingDisplayView);
+						bDisplayViewFound = true;
+						break;
+					}
+				}
 
-			CreateColorTransform(TopColorSpace.ColorSpaceName, OtherColorSpace.ColorSpaceName);
-			CreateColorTransform(OtherColorSpace.ColorSpaceName, TopColorSpace.ColorSpaceName);
-		}
+				if (!bDisplayViewFound)
+				{
+					UE_LOG(LogOpenColorIO, Display, TEXT("Removing %s, not found."), *ExistingDisplayView.ToString());
+				}
+			}
 
-		for (const FOpenColorIODisplayView& DisplayView : DesiredDisplayViews)
-		{
-			CreateColorTransform(TopColorSpace.ColorSpaceName, DisplayView.Display, DisplayView.View, EOpenColorIOViewTransformDirection::Forward);
+			const UOpenColorIOSettings* Settings = GetDefault<UOpenColorIOSettings>();
 
-			if (Settings->bSupportInverseViewTransforms)
+			// Genereate new shaders.
+			for (int32 indexTop = 0; indexTop < DesiredColorSpaces.Num(); ++indexTop)
 			{
-				CreateColorTransform(TopColorSpace.ColorSpaceName, DisplayView.Display, DisplayView.View, EOpenColorIOViewTransformDirection::Inverse);
+				const FOpenColorIOColorSpace& TopColorSpace = DesiredColorSpaces[indexTop];
+
+				for (int32 indexOther = indexTop + 1; indexOther < DesiredColorSpaces.Num(); ++indexOther)
+				{
+					const FOpenColorIOColorSpace& OtherColorSpace = DesiredColorSpaces[indexOther];
+
+					CreateColorTransform(TopColorSpace.ColorSpaceName, OtherColorSpace.ColorSpaceName);
+					CreateColorTransform(OtherColorSpace.ColorSpaceName, TopColorSpace.ColorSpaceName);
+				}
+
+				for (const FOpenColorIODisplayView& DisplayView : DesiredDisplayViews)
+				{
+					CreateColorTransform(TopColorSpace.ColorSpaceName, DisplayView.Display, DisplayView.View, EOpenColorIOViewTransformDirection::Forward);
+
+					if (Settings->bSupportInverseViewTransforms)
+					{
+						CreateColorTransform(TopColorSpace.ColorSpaceName, DisplayView.Display, DisplayView.View, EOpenColorIOViewTransformDirection::Inverse);
+					}
+				}
 			}
 		}
+		else
+		{
+			UE_LOG(LogOpenColorIO, Verbose, TEXT("Reload skipped, no differences identified."));
+		}
+	}
+	else
+	{
+		DesiredColorSpaces.Reset();
+		DesiredDisplayViews.Reset();
+		ColorTransforms.Reset();
+	}
+
+	if (!MarkPackageDirty())
+	{
+		UE_LOG(LogOpenColorIO, Display, TEXT("%s should be resaved for faster loads. Changes in the config or the library version were detected."), *GetPathName());
 	}
 #endif
 }
@@ -341,6 +372,7 @@ const TObjectPtr<UOpenColorIOColorTransform>* UOpenColorIOConfiguration::FindTra
 		});
 }
 
+#if WITH_EDITOR
 void UOpenColorIOConfiguration::CreateColorTransform(const FString& InSourceColorSpace, const FString& InDestinationColorSpace)
 {
 	if (InSourceColorSpace.IsEmpty() || InDestinationColorSpace.IsEmpty())
@@ -392,6 +424,7 @@ void UOpenColorIOConfiguration::CreateColorTransform(const FString& InSourceColo
 		UE_LOG(LogOpenColorIO, Warning, TEXT("Could not create color space transform from %s to %s - %s. Verify your OCIO config file, it may have errors in it."), *InSourceColorSpace, *InDisplay, *InView);
 	}
 }
+#endif
 
 void UOpenColorIOConfiguration::CleanupTransforms()
 {
@@ -442,13 +475,14 @@ void UOpenColorIOConfiguration::PostInitProperties()
 {
 	Super::PostInitProperties();
 
+#if WITH_EDITOR
 	if (!HasAnyFlags(RF_NeedLoad | RF_ClassDefaultObject))
 	{
 		ConfigurationFile.FilePath = TEXT("ocio://default");
-
 		// Ensure the default built-in configuration is loaded.
 		LoadConfiguration();
 	}
+#endif //WITH_EDITOR
 }
 
 void UOpenColorIOConfiguration::PostLoad()
@@ -584,11 +618,9 @@ void UOpenColorIOConfiguration::PostEditChangeProperty(FPropertyChangedEvent& Pr
 	}
 	Super::PostEditChangeProperty(PropertyChangedEvent);
 }
-#endif //WITH_EDITOR
 
 void UOpenColorIOConfiguration::LoadConfiguration()
 {
-#if WITH_EDITOR
 	Config.Reset();
 
 	if (ConfigurationFile.FilePath.IsEmpty())
@@ -628,7 +660,6 @@ void UOpenColorIOConfiguration::LoadConfiguration()
 	{
 		UE_LOG(LogOpenColorIO, Error, TEXT("Could not load OCIO configuration file %s. Verify that the path is good or that the file is valid."), *FilePath);
 	}
-#endif
 }
 
 void UOpenColorIOConfiguration::OnToastCallback(bool bInReloadColorspaces)
@@ -642,11 +673,13 @@ void UOpenColorIOConfiguration::OnToastCallback(bool bInReloadColorspaces)
 
 	if (bInReloadColorspaces)
 	{
+		OpenColorIOWrapper::ClearAllCaches();
+
 		ReloadExistingColorspaces();
 	}
 }
 
-#if WITH_EDITOR
+
 FOpenColorIOEditorConfigurationInspector::FOpenColorIOEditorConfigurationInspector(const UOpenColorIOConfiguration& InConfiguration)
 	: Configuration(InConfiguration)
 {
diff --git a/Engine/Plugins/Compositing/OpenColorIO/Source/OpenColorIO/Public/OpenColorIOColorTransform.h b/Engine/Plugins/Compositing/OpenColorIO/Source/OpenColorIO/Public/OpenColorIOColorTransform.h
--- a/Engine/Plugins/Compositing/OpenColorIO/Source/OpenColorIO/Public/OpenColorIOColorTransform.h
+++ b/Engine/Plugins/Compositing/OpenColorIO/Source/OpenColorIO/Public/OpenColorIOColorTransform.h
@@ -38,29 +38,23 @@ public:
 	UE_DEPRECATED(5.3, "This method is deprecated, please use Initialize without the owner argument.")
 	bool Initialize(UOpenColorIOConfiguration* InOwner, const FString& InSourceColorSpace, const FString& InDisplay, const FString& InView, EOpenColorIOViewTransformDirection InDirection, const TMap<FString, FString>& InContextKeyValues = {});
 
-	/** Initialize resources for color space transform. */
+#if WITH_EDITOR
+	/**
+	 * Initialize resources for color space transform. */
 	bool Initialize(const FString& InSourceColorSpace, const FString& InDestinationColorSpace, const TMap<FString, FString>& InContextKeyValues = {});
+	
 	/** Initialize resources for display-view transform. */
 	bool Initialize(const FString& InSourceColorSpace, const FString& InDisplay, const FString& InView, EOpenColorIOViewTransformDirection InDirection, const TMap<FString, FString>& InContextKeyValues = {});
-
-	/**
-	 * Cache resource shaders for cooking on the given shader platform.
-	 * If a matching shader map is not found in memory or the DDC, a new one will be compiled.
-	 * This does not apply completed results to the renderer scenes.
-	 * Caller is responsible for deleting OutCachedResource.
-	 */
-	void CacheResourceShadersForCooking(EShaderPlatform InShaderPlatform, const ITargetPlatform* TargetPlatform, const FString& InShaderHash, const FString& InShaderCode, const FString& InRawConfigHash, TArray<FOpenColorIOTransformResource*>& OutCachedResources);
+#endif
 
 	/**
 	 * Serialize LUT data. This will effectively serialize the LUT only when cooking
 	 */
-	void SerializeLuts(FArchive& Ar);
+	UE_DEPRECATED(5.3, "This method is deprecated.")
+	void SerializeLuts(FArchive& Ar) { }
 
-	/**
-	 * Cache resource LUT for rendering/cooking on the given shader platform.
-	 * If a matching texture is not found in DDC, a new one will be generated from raw data.
-	 */
-	void CacheResourceTextures();
+	UE_DEPRECATED(5.3, "This method is deprecated.")
+	void CacheResourceTextures() { }
 
 	/**
 	 * Cache resource shaders for rendering.
@@ -69,6 +63,7 @@ public:
 	void CacheResourceShadersForRendering(bool bRegenerateId);
 	void CacheShadersForResources(EShaderPlatform InShaderPlatform, FOpenColorIOTransformResource* InResourcesToCache, bool bApplyCompletedShaderMapForRendering, bool bIsCooking, const ITargetPlatform* TargetPlatform = nullptr);
 
+	UE_DEPRECATED(5.3, "This method is deprecated.")
 	FOpenColorIOTransformResource* AllocateResource();
 
 	/**
@@ -86,10 +81,10 @@ public:
 
 #if WITH_EDITOR
 	/** Apply the color transform in-place to the specified image. */
-	bool EditorTransformImage(const FImageView& InOutImage) const;
+	bool TransformImage(const FImageView& InOutImage) const;
 	
 	/** Apply the color transform from the source image to the destination image. (The destination FImageView is const but what it points at is not.) */
-	bool EditorTransformImage(const FImageView& SrcImage, const FImageView& DestImage) const;
+	bool TransformImage(const FImageView& SrcImage, const FImageView& DestImage) const;
 #endif
 
 	/**
@@ -124,26 +119,21 @@ protected:
 	 */
 	static void GetOpenColorIOLUTKeyGuid(const FString& InProcessorIdentifier, const FName& InName, FGuid& OutLutGuid );
 
-	/**
-	 * Generate the LUT and shader associated to the desired color space transform.
-	 * @note: Only in editor
-	 * @return true if parameters are valid.
-	 */
-	bool GenerateColorTransformData(const FString& InSourceColorSpace, const FString& InDestinationColorSpace);
-	bool GenerateColorTransformData(const FString& InSourceColorSpace, const FString& InDisplay, const FString& InView, EOpenColorIOViewTransformDirection InDirection);
+	UE_DEPRECATED(5.3, "This method is deprecated.")
+	bool GenerateColorTransformData(const FString& InSourceColorSpace, const FString& InDestinationColorSpace) { return false; }
+
+	UE_DEPRECATED(5.3, "This method is deprecated.")
+	bool GenerateColorTransformData(const FString& InSourceColorSpace, const FString& InDisplay, const FString& InView, EOpenColorIOViewTransformDirection InDirection) { return false; }
+
+	UE_DEPRECATED(5.3, "This method is deprecated.")
+	bool UpdateShaderInfo(FString& OutShaderCodeHash, FString& OutShaderCode, FString& OutRawConfigHash);
 
 	/**
 	 * Helper function returning the color space transform name based on source and destination color spaces.
 	 */
 	FString GetTransformFriendlyName() const;
 
-	/**
-	 * Fetch shader code and hash from the OCIO library
-	 * @note: Uses library only in editor. 
-	 * @return: true if not in editor (shader has been serialized during cooking), true if shader could be generated from the library otherwise
-	 */
-	bool UpdateShaderInfo(FString& OutShaderCodeHash, FString& OutShaderCode, FString& OutRawConfigHash);
-
+#if WITH_EDITOR
 	/**
 	 * Helper function taking raw 3D LUT data coming from the library and initializing a UVolumeTexture with it.
 	 */
@@ -153,8 +143,16 @@ protected:
 	 * Helper function taking raw 1D LUT data coming from the library and initializing a UTexture with it.
 	 */
 	TObjectPtr<UTexture> CreateTexture1DLUT(const FString& InProcessorIdentifier, const FName& InName, uint32 InTextureWidth, uint32 InTextureHeight, TextureFilter InFilter, bool bRedChannelOnly, const float* InSourceData );
+#endif //WITH_EDITOR
 
 private:
+#if WITH_EDITOR
+	/**
+	 * Create the transform processor(s) and generate its resources.
+	 */
+	void ProcessTransform();
+#endif //WITH_EDITOR
+
 	void FlushResourceShaderMaps();
 
 public:
@@ -171,6 +169,14 @@ public:
 	virtual void FinishDestroy() override;
 
 #if WITH_EDITOR
+	/**
+	 * Cache resource shaders for cooking on the given shader platform.
+	 * If a matching shader map is not found in memory or the DDC, a new one will be compiled.
+	 * This does not apply completed results to the renderer scenes.
+	 * Caller is responsible for deleting OutCachedResource.
+	 */
+	void CacheResourceShadersForCooking(EShaderPlatform InShaderPlatform, const ITargetPlatform* TargetPlatform, const FString& InShaderHash, const FString& InShaderCode, const FString& InRawConfigHash, TArray<FOpenColorIOTransformResource*>& OutCachedResources);
+
 	virtual void BeginCacheForCookedPlatformData(const ITargetPlatform* TargetPlatform) override;
 	virtual bool IsCachedCookedPlatformDataLoaded(const ITargetPlatform* TargetPlatform) override;
 	virtual void ClearCachedCookedPlatformData(const ITargetPlatform *TargetPlatform) override;
@@ -214,9 +220,18 @@ public:
 private:
 
 	/** If the color space requires textures, this will contains the data to do the transform */
-	/** Note: This will be serialized when cooking. Otherwhise, it relies on raw data of the library and what's on DDC */
-	UPROPERTY(Transient)
+	UPROPERTY()
 	TMap<int32, TObjectPtr<UTexture>> Textures;
+
+#if WITH_EDITORONLY_DATA
+	/** Generated transform shader hash. */
+	UPROPERTY()
+	FString GeneratedShaderHash;
+
+	/** Generated transform shader function. */
+	UPROPERTY()
+	FString GeneratedShader;
+#endif
 	
 	/** Inline ColorTransform resources serialized from disk. To be processed on game thread in PostLoad. */
 	TArray<FOpenColorIOTransformResource> LoadedTransformResources;
diff --git a/Engine/Plugins/Compositing/OpenColorIO/Source/OpenColorIO/Public/OpenColorIOConfiguration.h b/Engine/Plugins/Compositing/OpenColorIO/Source/OpenColorIO/Public/OpenColorIOConfiguration.h
--- a/Engine/Plugins/Compositing/OpenColorIO/Source/OpenColorIO/Public/OpenColorIOConfiguration.h
+++ b/Engine/Plugins/Compositing/OpenColorIO/Source/OpenColorIO/Public/OpenColorIOConfiguration.h
@@ -47,10 +47,10 @@ public:
 
 #if WITH_EDITOR
 	/** Apply the color transform in-place to the specified image. */
-	bool EditorTransformImage(const FOpenColorIOColorConversionSettings& InSettings, const FImageView& InOutImage) const;
+	bool TransformImage(const FOpenColorIOColorConversionSettings& InSettings, const FImageView& InOutImage) const;
 
 	/** Apply the color transform from the source image to the destination image. (The destination FImageView is const but what it points at is not.) */
-	bool EditorTransformImage(const FOpenColorIOColorConversionSettings& InSettings, const FImageView& SrcImage, const FImageView& DestImage) const;
+	bool TransformImage(const FOpenColorIOColorConversionSettings& InSettings, const FImageView& SrcImage, const FImageView& DestImage) const;
 #endif
 
 	/** This forces to reload colorspaces and corresponding shaders if those are not loaded already. */
@@ -69,8 +69,10 @@ public:
 protected:
 
 	const TObjectPtr<UOpenColorIOColorTransform>* FindTransform(const FOpenColorIOColorConversionSettings& InSettings) const;
+#if WITH_EDITOR
 	void CreateColorTransform(const FString& InSourceColorSpace, const FString& InDestinationColorSpace);
 	void CreateColorTransform(const FString& InSourceColorSpace, const FString& InDisplay, const FString& InView, EOpenColorIOViewTransformDirection InDirection);
+#endif
 	void CleanupTransforms();
 
 	/** Same as above except user can specify the path manually. */
@@ -98,10 +100,12 @@ public:
 	//~ End UObject interface
 
 private:
+#if WITH_EDITOR
 	void LoadConfiguration();
 
 	/** This method resets the status of Notification dialog and reacts depending on user's choice. */
 	void OnToastCallback(bool bInReloadColorspaces);
+#endif
 
 public:
 
@@ -138,6 +142,10 @@ private:
 #if WITH_EDITORONLY_DATA
 	/** Private implementation of the OpenColorIO config object. */
 	TPimplPtr<FOpenColorIOWrapperConfig> Config = nullptr;
+
+	/** Hash of all of the config content, including relevant external file information. */
+	UPROPERTY()
+	FString ConfigHash;
 #endif
 };
 
diff --git a/Engine/Plugins/Compositing/OpenColorIO/Source/OpenColorIOEditor/Private/OpenColorIOConfigurationCustomization.cpp b/Engine/Plugins/Compositing/OpenColorIO/Source/OpenColorIOEditor/Private/OpenColorIOConfigurationCustomization.cpp
--- a/Engine/Plugins/Compositing/OpenColorIO/Source/OpenColorIOEditor/Private/OpenColorIOConfigurationCustomization.cpp
+++ b/Engine/Plugins/Compositing/OpenColorIO/Source/OpenColorIOEditor/Private/OpenColorIOConfigurationCustomization.cpp
@@ -4,6 +4,7 @@
 
 #include "OpenColorIOColorSpaceCustomization.h"
 #include "OpenColorIOConfiguration.h"
+#include "OpenColorIOWrapper.h"
 #include "Containers/StringConv.h"
 #include "DetailWidgetRow.h"
 #include "DetailLayoutBuilder.h"
@@ -46,6 +47,8 @@ void FOpenColorIOConfigurationCustomization::CustomizeDetails(IDetailLayoutBuild
 			{
 				if(UOpenColorIOConfiguration* OpenColorIOConfig = Cast<UOpenColorIOConfiguration>(Objects[0].Get()))
 				{
+					OpenColorIOWrapper::ClearAllCaches();
+
 					OpenColorIOConfig->ReloadExistingColorspaces();
 				}
 				return FReply::Handled();
diff --git a/Engine/Plugins/Compositing/OpenColorIO/Source/OpenColorIOEditor/Private/OpenColorIOEditorBlueprintLibrary.cpp b/Engine/Plugins/Compositing/OpenColorIO/Source/OpenColorIOEditor/Private/OpenColorIOEditorBlueprintLibrary.cpp
--- a/Engine/Plugins/Compositing/OpenColorIO/Source/OpenColorIOEditor/Private/OpenColorIOEditorBlueprintLibrary.cpp
+++ b/Engine/Plugins/Compositing/OpenColorIO/Source/OpenColorIOEditor/Private/OpenColorIOEditorBlueprintLibrary.cpp
@@ -37,7 +37,7 @@ bool UOpenColorIOEditorBlueprintLibrary::ApplyColorSpaceCompressionTransformToTe
 		FImage ImageMip0;
 		if (InOutTexture->Source.GetMipImage(ImageMip0, 0))
 		{
-			bool bTransformSucceeded = ConversionSettings.ConfigurationSource->EditorTransformImage(ConversionSettings, ImageMip0);
+			bool bTransformSucceeded = ConversionSettings.ConfigurationSource->TransformImage(ConversionSettings, ImageMip0);
 			if (bTransformSucceeded)
 			{
 				InOutTexture->PreEditChange(nullptr);
diff --git a/Engine/Plugins/Compositing/OpenColorIO/Source/OpenColorIOWrapper/Private/OpenColorIOWrapper.cpp b/Engine/Plugins/Compositing/OpenColorIO/Source/OpenColorIOWrapper/Private/OpenColorIOWrapper.cpp
--- a/Engine/Plugins/Compositing/OpenColorIO/Source/OpenColorIOWrapper/Private/OpenColorIOWrapper.cpp
+++ b/Engine/Plugins/Compositing/OpenColorIO/Source/OpenColorIOWrapper/Private/OpenColorIOWrapper.cpp
@@ -59,6 +59,13 @@ const TCHAR* OpenColorIOWrapper::GetVersion()
 #endif // WITH_OCIO
 }
 
+void OpenColorIOWrapper::ClearAllCaches()
+{
+#if WITH_OCIO
+	OCIO_NAMESPACE::ClearAllCaches();
+#endif
+}
+
 struct FOpenColorIOConfigPimpl
 {
 #if WITH_OCIO
@@ -867,6 +874,18 @@ bool FOpenColorIOWrapperProcessor::IsValid() const
 #endif // WITH_OCIO
 }
 
+FString FOpenColorIOWrapperProcessor::GetCacheID() const
+{
+#if WITH_OCIO
+	if (IsValid())
+	{
+		return StringCast<TCHAR>(Pimpl->Processor->getCacheID()).Get();
+	}
+#endif // WITH_OCIO
+
+	return {};
+}
+
 FString FOpenColorIOWrapperProcessor::GetTransformToWorkingColorSpaceName(const FOpenColorIOWrapperSourceColorSettings& InSourceColorSettings)
 {
 	const uint32 SettingsId = (uint32)InSourceColorSettings.EncodingOverride | (uint32)InSourceColorSettings.ColorSpace << 8u | (uint32)InSourceColorSettings.ChromaticAdaptationMethod << 16u;
@@ -1070,15 +1089,15 @@ bool FOpenColorIOWrapperGPUProcessor::IsValid() const
 #endif // WITH_OCIO
 }
 
-bool FOpenColorIOWrapperGPUProcessor::GetShader(FString& OutShaderCode, FString& OutShaderCacheID) const
+bool FOpenColorIOWrapperGPUProcessor::GetShader(FString& OutShaderCacheID, FString& OutShaderCode) const
 {
 #if WITH_OCIO
 	if (IsValid())
 	{
 		ensureMsgf(GPUPimpl->ShaderDescription->getNumDynamicProperties() == 0, TEXT("We do not currently support dynamic properties."));
 
-		OutShaderCode = StringCast<TCHAR>(GPUPimpl->ShaderDescription->getShaderText()).Get();
 		OutShaderCacheID = StringCast<TCHAR>(GPUPimpl->ShaderDescription->getCacheID()).Get();
+		OutShaderCode = StringCast<TCHAR>(GPUPimpl->ShaderDescription->getShaderText()).Get();
 
 		return true;
 	}
diff --git a/Engine/Plugins/Compositing/OpenColorIO/Source/OpenColorIOWrapper/Public/OpenColorIOWrapper.h b/Engine/Plugins/Compositing/OpenColorIO/Source/OpenColorIOWrapper/Public/OpenColorIOWrapper.h
--- a/Engine/Plugins/Compositing/OpenColorIO/Source/OpenColorIOWrapper/Public/OpenColorIOWrapper.h
+++ b/Engine/Plugins/Compositing/OpenColorIO/Source/OpenColorIOWrapper/Public/OpenColorIOWrapper.h
@@ -34,6 +34,9 @@ namespace OpenColorIOWrapper
 
 	/** Get the OpenColorIO version string. */
 	OPENCOLORIOWRAPPER_API const TCHAR* GetVersion();
+
+	/** Calls the native function of the same name. */
+	OPENCOLORIOWRAPPER_API void ClearAllCaches();
 }
 
 
@@ -208,6 +211,9 @@ public:
 	/** Valid when the processor has been successfully created and isn't null. */
 	OPENCOLORIOWRAPPER_API bool IsValid() const;
 
+	/** Get the string hash of the processor. */
+	OPENCOLORIOWRAPPER_API FString GetCacheID() const;
+
 	/**
 	* Get the generated transform name from source color settings to working color space.
 	* @param InColorSettings Source encoding and colorspace.
@@ -259,11 +265,11 @@ public:
 	/**
 	* Get the generated shader.
 	*
-	* @param OutShaderCode Generated shader text.
 	* @param OutShaderCacheID Generated shader hash string.
+	* @param OutShaderCode Generated shader text.
 	* @return True when the result is valid.
 	*/
-	OPENCOLORIOWRAPPER_API bool GetShader(FString& OutShaderCode, FString& OutShaderCacheID) const;
+	OPENCOLORIOWRAPPER_API bool GetShader(FString& OutShaderCacheID, FString& OutShaderCode) const;
 
 	/** Get the number of 3D LUT textures. */
 	OPENCOLORIOWRAPPER_API uint32 GetNum3DTextures() const;
diff --git a/Engine/Source/Runtime/Core/Public/UObject/UE5MainStreamObjectVersion.h b/Engine/Source/Runtime/Core/Public/UObject/UE5MainStreamObjectVersion.h
--- a/Engine/Source/Runtime/Core/Public/UObject/UE5MainStreamObjectVersion.h
+++ b/Engine/Source/Runtime/Core/Public/UObject/UE5MainStreamObjectVersion.h
@@ -363,6 +363,9 @@ struct CORE_API FUE5MainStreamObjectVersion
 
 		// Cloth collection tether initialization
 		ClothCollectionTetherInitialization,
+
+		// OpenColorIO transforms now serialize their generated texture(s) and shader code normally into the uasset.
+		OpenColorIOAssetCacheSerialization,
 		
 		// -----<new versions can be added above this line>-------------------------------------------------
 		VersionPlusOne,
