OCIO: Handle empty generated shader property while cooking.

diff --git a/Engine/Plugins/Compositing/OpenColorIO/Source/OpenColorIO/Private/OpenColorIOColorTransform.cpp b/Engine/Plugins/Compositing/OpenColorIO/Source/OpenColorIO/Private/OpenColorIOColorTransform.cpp
--- a/Engine/Plugins/Compositing/OpenColorIO/Source/OpenColorIO/Private/OpenColorIOColorTransform.cpp
+++ b/Engine/Plugins/Compositing/OpenColorIO/Source/OpenColorIO/Private/OpenColorIOColorTransform.cpp
@@ -710,38 +710,6 @@ FString UOpenColorIOColorTransform::GetTransformFriendlyName() const
 	}
 }
 
-bool UOpenColorIOColorTransform::UpdateShaderInfo(FString& OutShaderCodeHash, FString& OutShaderCode, FString& OutRawConfigHash)
-{
-#if WITH_EDITOR
-	const UOpenColorIOSettings* Settings = GetDefault<UOpenColorIOSettings>();
-	const UOpenColorIOConfiguration* ConfigurationOwner = Cast<UOpenColorIOConfiguration>(GetOuter());
-	const FOpenColorIOWrapperConfig* ConfigWrapper = GetTransformConfigWrapper(ConfigurationOwner);
-
-	if (ConfigWrapper != nullptr)
-	{
-		const FOpenColorIOWrapperProcessor Processor = GetTransformProcessor(this, ConfigWrapper);
-		if (Processor.IsValid())
-		{
-			const FOpenColorIOWrapperGPUProcessor GPUProcessor = FOpenColorIOWrapperGPUProcessor(Processor, Settings->bUseLegacyProcessor);
-			if (GPUProcessor.IsValid())
-			{
-				OutRawConfigHash = ConfigurationOwner->GetConfigWrapper()->GetCacheID();
-
-				return GPUProcessor.GetShader(OutShaderCodeHash, OutShaderCode);
-			}
-		}
-	}
-	else
-	{
-		UE_LOG(LogOpenColorIO, Error, TEXT("Failed to fetch shader info for color transform %s. Configuration file [%s] was invalid."), *GetTransformFriendlyName(), ConfigurationOwner ? *ConfigurationOwner->ConfigurationFile.FilePath : TEXT("Missing"));
-	}
-
-	return false;
-#else
-	return true; //When not in editor, shaders have been cooked so we're not relying on the library data anymore.
-#endif //WITH_EDITOR
-}
-
 void UOpenColorIOColorTransform::FlushResourceShaderMaps()
 {
 	if (FApp::CanEverRender())
@@ -851,6 +819,32 @@ void UOpenColorIOColorTransform::FinishDestroy()
 
 #if WITH_EDITOR
 
+bool UOpenColorIOColorTransform::UpdateShaderInfo(FString& OutShaderCodeHash, FString& OutShaderCode, FString& OutRawConfigHash)
+{
+	const UOpenColorIOSettings* Settings = GetDefault<UOpenColorIOSettings>();
+	const UOpenColorIOConfiguration* ConfigurationOwner = Cast<UOpenColorIOConfiguration>(GetOuter());
+	const FOpenColorIOWrapperConfig* ConfigWrapper = GetTransformConfigWrapper(ConfigurationOwner);
+
+	if (ConfigWrapper != nullptr)
+	{
+		const FOpenColorIOWrapperProcessor Processor = GetTransformProcessor(this, ConfigWrapper);
+		if (Processor.IsValid())
+		{
+			const FOpenColorIOWrapperGPUProcessor GPUProcessor = FOpenColorIOWrapperGPUProcessor(Processor, Settings->bUseLegacyProcessor);
+			if (GPUProcessor.IsValid())
+			{
+				OutRawConfigHash = ConfigurationOwner->GetConfigWrapper()->GetCacheID();
+
+				return GPUProcessor.GetShader(OutShaderCodeHash, OutShaderCode);
+			}
+		}
+	}
+
+	UE_LOG(LogOpenColorIO, Error, TEXT("Failed to fetch shader info for color transform %s. Configuration file [%s] was invalid."), *GetTransformFriendlyName(), ConfigurationOwner ? *ConfigurationOwner->ConfigurationFile.FilePath : TEXT("Missing"));
+
+	return false;
+}
+
 void UOpenColorIOColorTransform::BeginCacheForCookedPlatformData(const ITargetPlatform* TargetPlatform)
 {
 	TArray<FName> DesiredShaderFormats;
@@ -860,16 +854,32 @@ void UOpenColorIOColorTransform::BeginCacheForCookedPlatformData(const ITargetPl
 
 	if (DesiredShaderFormats.Num() > 0)
 	{
-		const UOpenColorIOConfiguration* ConfigurationOwner = Cast<UOpenColorIOConfiguration>(GetOuter());
-		if (ConfigurationOwner != nullptr)
+		bool bValidShaderInfo = false;
+		FString ShaderCodeHash;
+		FString ShaderCode;
+		FString RawConfigHash;
+
+		if (GeneratedShader.IsEmpty())
+		{
+			//Need to re-update shader data when cooking, they may not have been previously fetched.
+			bValidShaderInfo = UpdateShaderInfo(ShaderCodeHash, ShaderCode, RawConfigHash);
+		}
+		else if(const UOpenColorIOConfiguration* ConfigurationOwner = Cast<UOpenColorIOConfiguration>(GetOuter()))
+		{
+			ShaderCodeHash = GeneratedShaderHash;
+			ShaderCode = GeneratedShader;
+			RawConfigHash = ConfigurationOwner->GetConfigWrapper()->GetCacheID();
+			bValidShaderInfo = true;
+		}
+
+		if (bValidShaderInfo)
 		{
-			FString RawConfigHash = ConfigurationOwner->GetConfigWrapper()->GetCacheID();
 			// Cache for all the shader formats that the cooking target requires
 			for (int32 FormatIndex = 0; FormatIndex < DesiredShaderFormats.Num(); FormatIndex++)
 			{
 				const EShaderPlatform LegacyShaderPlatform = ShaderFormatToLegacyShaderPlatform(DesiredShaderFormats[FormatIndex]);
 				// Begin caching shaders for the target platform and store the FOpenColorIOTransformResource being compiled into CachedColorTransformResourcesForCooking
-				CacheResourceShadersForCooking(LegacyShaderPlatform, TargetPlatform, GeneratedShaderHash, GeneratedShader, RawConfigHash, *CachedColorTransformResourceForPlatformPtr);
+				CacheResourceShadersForCooking(LegacyShaderPlatform, TargetPlatform, ShaderCodeHash, ShaderCode, RawConfigHash, *CachedColorTransformResourceForPlatformPtr);
 			}
 		}
 	}
diff --git a/Engine/Plugins/Compositing/OpenColorIO/Source/OpenColorIO/Public/OpenColorIOColorTransform.h b/Engine/Plugins/Compositing/OpenColorIO/Source/OpenColorIO/Public/OpenColorIOColorTransform.h
--- a/Engine/Plugins/Compositing/OpenColorIO/Source/OpenColorIO/Public/OpenColorIOColorTransform.h
+++ b/Engine/Plugins/Compositing/OpenColorIO/Source/OpenColorIO/Public/OpenColorIOColorTransform.h
@@ -125,15 +125,18 @@ protected:
 	UE_DEPRECATED(5.3, "This method is deprecated.")
 	bool GenerateColorTransformData(const FString& InSourceColorSpace, const FString& InDisplay, const FString& InView, EOpenColorIOViewTransformDirection InDirection) { return false; }
 
-	UE_DEPRECATED(5.3, "This method is deprecated.")
-	bool UpdateShaderInfo(FString& OutShaderCodeHash, FString& OutShaderCode, FString& OutRawConfigHash);
-
 	/**
 	 * Helper function returning the color space transform name based on source and destination color spaces.
 	 */
 	FString GetTransformFriendlyName() const;
 
 #if WITH_EDITOR
+	/**
+	 * Fetch shader code and hash from the OCIO library
+	 * @return: true if shader could be generated from the library
+	 */
+	bool UpdateShaderInfo(FString& OutShaderCodeHash, FString& OutShaderCode, FString& OutRawConfigHash);
+
 	/**
 	 * Helper function taking raw 3D LUT data coming from the library and initializing a UVolumeTexture with it.
 	 */
