OCIO: Terminate macro calls with semicolons.

diff --git a/Engine/Plugins/Compositing/OpenColorIO/Source/OpenColorIOWrapper/Private/OpenColorIOWrapper.cpp b/Engine/Plugins/Compositing/OpenColorIO/Source/OpenColorIOWrapper/Private/OpenColorIOWrapper.cpp
--- a/Engine/Plugins/Compositing/OpenColorIO/Source/OpenColorIOWrapper/Private/OpenColorIOWrapper.cpp
+++ b/Engine/Plugins/Compositing/OpenColorIO/Source/OpenColorIOWrapper/Private/OpenColorIOWrapper.cpp
@@ -151,8 +151,8 @@ FOpenColorIOWrapperConfig::FOpenColorIOWrapperConfig(FStringView InFilePath, EAd
 	: FOpenColorIOWrapperConfig()
 {
 #if WITH_OCIO
-	TRACE_CPUPROFILER_EVENT_SCOPE(FOpenColorIOWrapperConfig::FOpenColorIOWrapperConfig)
-	OCIO_EXCEPTION_HANDLING_TRY()
+	TRACE_CPUPROFILER_EVENT_SCOPE(FOpenColorIOWrapperConfig::FOpenColorIOWrapperConfig);
+	OCIO_EXCEPTION_HANDLING_TRY();
 
 		using namespace OCIO_NAMESPACE;
 
@@ -187,7 +187,7 @@ FOpenColorIOWrapperConfig::FOpenColorIOWrapperConfig(FStringView InFilePath, EAd
 
 		Pimpl->Config = NewConfig;
 	
-	OCIO_EXCEPTION_HANDLING_CATCH(Error, TEXT("Could not create OCIO configuration file for %s. Error message: %s"), InFilePath.GetData())
+	OCIO_EXCEPTION_HANDLING_CATCH(Error, TEXT("Could not create OCIO configuration file for %s. Error message: %s"), InFilePath.GetData());
 #else
 	UE_LOG(LogOpenColorIOWrapper, Error, TEXT("The OpenColorIO library is not available, functionality will be disabled."));
 #endif // WITH_OCIO
@@ -439,7 +439,7 @@ FOpenColorIOWrapperProcessor::FOpenColorIOWrapperProcessor(
 	}
 
 #if WITH_OCIO
-	OCIO_EXCEPTION_HANDLING_TRY()
+	OCIO_EXCEPTION_HANDLING_TRY();
 		if (OwnerConfig != nullptr && OwnerConfig->IsValid())
 		{
 			const OCIO_NAMESPACE::ConstConfigRcPtr& Config = OwnerConfig->Pimpl->Config;
@@ -456,7 +456,7 @@ FOpenColorIOWrapperProcessor::FOpenColorIOWrapperProcessor(
 				StringCast<ANSICHAR>(InDestinationColorSpace.GetData()).Get()
 			);
 		}
-	OCIO_EXCEPTION_HANDLING_CATCH(Log, TEXT("Failed to create processor for [%s, %s]. Error message: %s"), InSourceColorSpace.GetData(), InDestinationColorSpace.GetData())
+	OCIO_EXCEPTION_HANDLING_CATCH(Log, TEXT("Failed to create processor for [%s, %s]. Error message: %s"), InSourceColorSpace.GetData(), InDestinationColorSpace.GetData());
 #endif // WITH_OCIO
 }
 
@@ -484,7 +484,7 @@ FOpenColorIOWrapperProcessor::FOpenColorIOWrapperProcessor(
 	}
 
 #if WITH_OCIO
-	OCIO_EXCEPTION_HANDLING_TRY()
+	OCIO_EXCEPTION_HANDLING_TRY();
 		if (OwnerConfig != nullptr && OwnerConfig->IsValid())
 		{
 			const OCIO_NAMESPACE::ConstConfigRcPtr& Config = OwnerConfig->Pimpl->Config;
@@ -512,7 +512,7 @@ FOpenColorIOWrapperProcessor::FOpenColorIOWrapperProcessor(const FOpenColorIOWra
 	, WorkingColorSpaceTransformType(EOpenColorIOWorkingColorSpaceTransform::None)
 {
 #if WITH_OCIO
-	OCIO_EXCEPTION_HANDLING_TRY()
+	OCIO_EXCEPTION_HANDLING_TRY();
 		if (OwnerConfig != nullptr && OwnerConfig->IsValid())
 		{
 			const OCIO_NAMESPACE::ConstConfigRcPtr& Config = OwnerConfig->Pimpl->Config;
@@ -907,8 +907,8 @@ FString FOpenColorIOWrapperProcessor::GetTransformToWorkingColorSpaceName(const
 bool FOpenColorIOWrapperProcessor::TransformColor(FLinearColor& InOutColor) const
 {
 #if WITH_OCIO
-	TRACE_CPUPROFILER_EVENT_SCOPE(FOpenColorIOWrapperProcessor::TransformColor)
-	OCIO_EXCEPTION_HANDLING_TRY()
+	TRACE_CPUPROFILER_EVENT_SCOPE(FOpenColorIOWrapperProcessor::TransformColor);
+	OCIO_EXCEPTION_HANDLING_TRY();
 
 		if (IsValid())
 		{
@@ -918,7 +918,7 @@ bool FOpenColorIOWrapperProcessor::TransformColor(FLinearColor& InOutColor) cons
 
 			return true;
 		}
-	OCIO_EXCEPTION_HANDLING_CATCH(Log, TEXT("Failed to transform color. Error message: %s"))
+	OCIO_EXCEPTION_HANDLING_CATCH(Log, TEXT("Failed to transform color. Error message: %s"));
 #endif // WITH_OCIO
 
 	return false;
@@ -927,8 +927,8 @@ bool FOpenColorIOWrapperProcessor::TransformColor(FLinearColor& InOutColor) cons
 bool FOpenColorIOWrapperProcessor::TransformImage(const FImageView& InOutImage) const
 {
 #if WITH_OCIO
-	TRACE_CPUPROFILER_EVENT_SCOPE(FOpenColorIOWrapperProcessor::TransformImage)
-	OCIO_EXCEPTION_HANDLING_TRY()
+	TRACE_CPUPROFILER_EVENT_SCOPE(FOpenColorIOWrapperProcessor::TransformImage);
+	OCIO_EXCEPTION_HANDLING_TRY();
 		using namespace OCIO_NAMESPACE;
 
 		if (IsValid())
@@ -975,7 +975,7 @@ bool FOpenColorIOWrapperProcessor::TransformImage(const FImageView& InOutImage)
 				return true;
 			}
 		}
-	OCIO_EXCEPTION_HANDLING_CATCH(Log, TEXT("Failed to transform image. Error message: %s"))
+	OCIO_EXCEPTION_HANDLING_CATCH(Log, TEXT("Failed to transform image. Error message: %s"));
 #endif // WITH_OCIO
 
 	return false;
@@ -984,8 +984,8 @@ bool FOpenColorIOWrapperProcessor::TransformImage(const FImageView& InOutImage)
 bool FOpenColorIOWrapperProcessor::TransformImage(const FImageView& SrcImage, const FImageView& DestImage) const
 {
 #if WITH_OCIO
-	TRACE_CPUPROFILER_EVENT_SCOPE(FOpenColorIOWrapperProcessor::TransformImage)
-	OCIO_EXCEPTION_HANDLING_TRY()
+	TRACE_CPUPROFILER_EVENT_SCOPE(FOpenColorIOWrapperProcessor::TransformImage);
+	OCIO_EXCEPTION_HANDLING_TRY();
 		using namespace OCIO_NAMESPACE;
 
 		if (IsValid())
@@ -1035,7 +1035,7 @@ bool FOpenColorIOWrapperProcessor::TransformImage(const FImageView& SrcImage, co
 				return true;
 			}
 		}
-	OCIO_EXCEPTION_HANDLING_CATCH(Log, TEXT("Failed to transform image. Error message: %s"))
+	OCIO_EXCEPTION_HANDLING_CATCH(Log, TEXT("Failed to transform image. Error message: %s"));
 #endif // WITH_OCIO
 
 	return false;
@@ -1048,8 +1048,8 @@ FOpenColorIOWrapperGPUProcessor::FOpenColorIOWrapperGPUProcessor(FOpenColorIOWra
 	, GPUPimpl(MakePimpl<FOpenColorIOGPUProcessorPimpl, EPimplPtrMode::DeepCopy>())
 {
 #if WITH_OCIO
-	TRACE_CPUPROFILER_EVENT_SCOPE(FOpenColorIOWrapperGPUProcessor::FOpenColorIOWrapperGPUProcessor)
-	OCIO_EXCEPTION_HANDLING_TRY()
+	TRACE_CPUPROFILER_EVENT_SCOPE(FOpenColorIOWrapperGPUProcessor::FOpenColorIOWrapperGPUProcessor);
+	OCIO_EXCEPTION_HANDLING_TRY();
 		using namespace OCIO_NAMESPACE;
 
 		if (ParentProcessor.IsValid())
@@ -1076,7 +1076,7 @@ FOpenColorIOWrapperGPUProcessor::FOpenColorIOWrapperGPUProcessor(FOpenColorIOWra
 			GPUPimpl->Processor = GPUProcessor;
 			GPUPimpl->ShaderDescription = ShaderDescription;
 		}
-	OCIO_EXCEPTION_HANDLING_CATCH(Log, TEXT("Failed to fetch shader info for color transform. Error message: %s"))
+	OCIO_EXCEPTION_HANDLING_CATCH(Log, TEXT("Failed to fetch shader info for color transform. Error message: %s"));
 #endif // WITH_OCIO
 }
 
@@ -1122,7 +1122,7 @@ bool FOpenColorIOWrapperGPUProcessor::Get3DTexture(uint32 InIndex, FName& OutNam
 {
 #if WITH_OCIO
 	ensure(InIndex < GetNum3DTextures());
-	OCIO_EXCEPTION_HANDLING_TRY()
+	OCIO_EXCEPTION_HANDLING_TRY();
 		if (IsValid())
 		{
 			const ANSICHAR* TextureName = nullptr;
@@ -1146,7 +1146,7 @@ bool FOpenColorIOWrapperGPUProcessor::Get3DTexture(uint32 InIndex, FName& OutNam
 
 			return TextureName && OutEdgeLength > 0 && OutData;
 		}
-	OCIO_EXCEPTION_HANDLING_CATCH(Log, TEXT("Failed to fetch 3d texture(s) info for color transform. Error message: %s"))
+	OCIO_EXCEPTION_HANDLING_CATCH(Log, TEXT("Failed to fetch 3d texture(s) info for color transform. Error message: %s"));
 #endif // WITH_OCIO
 
 	return false;
@@ -1168,7 +1168,7 @@ bool FOpenColorIOWrapperGPUProcessor::GetTexture(uint32 InIndex, FName& OutName,
 	ensure(InIndex < GetNumTextures());
 
 #if WITH_OCIO
-	OCIO_EXCEPTION_HANDLING_TRY()
+	OCIO_EXCEPTION_HANDLING_TRY();
 		if (IsValid())
 		{
 			const ANSICHAR* TextureName = nullptr;
@@ -1189,7 +1189,7 @@ bool FOpenColorIOWrapperGPUProcessor::GetTexture(uint32 InIndex, FName& OutName,
 
 			return TextureName && OutWidth > 0 && OutHeight > 0 && OutData;
 		}
-	OCIO_EXCEPTION_HANDLING_CATCH(Error, TEXT("Failed to fetch texture(s) info for color transform. Error message: %s"))
+	OCIO_EXCEPTION_HANDLING_CATCH(Error, TEXT("Failed to fetch texture(s) info for color transform. Error message: %s"));
 #endif // WITH_OCIO
 
 	return false;
