OCIO: Added new universal method for applying OCIO transforms and shuffled a few things around.

diff --git a/Engine/Plugins/Compositing/OpenColorIO/Source/OpenColorIO/OpenColorIO.Build.cs b/Engine/Plugins/Compositing/OpenColorIO/Source/OpenColorIO/OpenColorIO.Build.cs
--- a/Engine/Plugins/Compositing/OpenColorIO/Source/OpenColorIO/OpenColorIO.Build.cs
+++ b/Engine/Plugins/Compositing/OpenColorIO/Source/OpenColorIO/OpenColorIO.Build.cs
@@ -19,7 +19,6 @@ namespace UnrealBuildTool.Rules
 			PrivateDependencyModuleNames.AddRange(
 				new string[] {
 					"Projects",
-					"RenderCore",
 					"RHI",
 					"RenderCore",
 					"Renderer",
diff --git a/Engine/Plugins/Compositing/OpenColorIO/Source/OpenColorIO/Private/OpenColorIODisplayExtension.cpp b/Engine/Plugins/Compositing/OpenColorIO/Source/OpenColorIO/Private/OpenColorIODisplayExtension.cpp
--- a/Engine/Plugins/Compositing/OpenColorIO/Source/OpenColorIO/Private/OpenColorIODisplayExtension.cpp
+++ b/Engine/Plugins/Compositing/OpenColorIO/Source/OpenColorIO/Private/OpenColorIODisplayExtension.cpp
@@ -20,7 +20,6 @@
 // for FPostProcessMaterialInputs
 #include "PostProcess/PostProcessMaterialInputs.h"
 
-float FOpenColorIODisplayExtension::DefaultDisplayGamma = 2.2f;
 
 FOpenColorIODisplayExtension::FOpenColorIODisplayExtension(const FAutoRegister& AutoRegister, FViewportClient* AssociatedViewportClient)
 	: FSceneViewExtensionBase(AutoRegister)
@@ -55,15 +54,8 @@ void FOpenColorIODisplayExtension::SetupView(FSceneViewFamily& InViewFamily, FSc
 
 	if (PassResources.IsValid())
 	{
-		//Force ToneCurve to be off while we'are alive to make sure the input color space is the working space : srgb linear
-		InViewFamily.EngineShowFlags.SetToneCurve(false);
-		// This flags sets tonampper to output to ETonemapperOutputDevice::LinearNoToneCurve
-		InViewFamily.SceneCaptureSource = SCS_FinalColorHDR;
-
-		InView.FinalPostProcessSettings.bOverride_ToneCurveAmount = 1;
-		InView.FinalPostProcessSettings.ToneCurveAmount = 0.0;
+		FOpenColorIORendering::PrepareView(InViewFamily, InView);
 	}
-
 	ENQUEUE_RENDER_COMMAND(ProcessColorSpaceTransform)(
 		[this, ResourcesRenderThread = MoveTemp(PassResources)](FRHICommandListImmediate& RHICmdList)
 		{
@@ -94,19 +86,12 @@ FScreenPassTexture FOpenColorIODisplayExtension::PostProcessPassAfterTonemap_Ren
 		Output = FScreenPassRenderTarget::CreateFromInput(GraphBuilder, SceneColor, View.GetOverwriteLoadAction(), TEXT("OCIORenderTarget"));
 	}
 
-	const float EngineDisplayGamma = View.Family->RenderTarget->GetDisplayGamma();
-	// There is a special case where post processing and tonemapper are disabled. In this case tonemapper applies a static display Inverse of Gamma which defaults to 2.2.
-	// In the case when Both PostProcessing and ToneMapper are disabled we apply gamma manually. In every other case we apply inverse gamma before applying OCIO.
-	float DisplayGamma = (View.Family->EngineShowFlags.Tonemapper == 0) || (View.Family->EngineShowFlags.PostProcessing == 0) ? DefaultDisplayGamma : DefaultDisplayGamma / EngineDisplayGamma;
-
 	FOpenColorIORendering::AddPass_RenderThread(
 		GraphBuilder,
 		View,
-		View.GetFeatureLevel(),
 		SceneColor,
 		Output,
-		CachedResourcesRenderThread,
-		DisplayGamma
+		CachedResourcesRenderThread
 	);
 
 	return MoveTemp(Output);
diff --git a/Engine/Plugins/Compositing/OpenColorIO/Source/OpenColorIO/Private/OpenColorIORendering.cpp b/Engine/Plugins/Compositing/OpenColorIO/Source/OpenColorIO/Private/OpenColorIORendering.cpp
--- a/Engine/Plugins/Compositing/OpenColorIO/Source/OpenColorIO/Private/OpenColorIORendering.cpp
+++ b/Engine/Plugins/Compositing/OpenColorIO/Source/OpenColorIO/Private/OpenColorIORendering.cpp
@@ -23,6 +23,8 @@
 #include "TextureResource.h"
 
 
+float FOpenColorIORendering::DefaultDisplayGamma = 2.2f;
+
 namespace {
 	using namespace UE::Color;
 
@@ -104,6 +106,43 @@ void FOpenColorIORendering::AddPass_RenderThread(
 	}
 }
 
+// static
+void FOpenColorIORendering::AddPass_RenderThread(
+	FRDGBuilder& GraphBuilder,
+	const FSceneView& View,
+	const FScreenPassTexture& Input,
+	const FScreenPassRenderTarget& Output,
+	const FOpenColorIORenderPassResources& InPassInfo,
+	EOpenColorIOTransformAlpha TransformAlpha) 
+{
+	const float EngineDisplayGamma = View.Family->RenderTarget->GetDisplayGamma();
+	// There is a special case where post processing and tonemapper are disabled. In this case tonemapper applies a static display Inverse of Gamma which defaults to 2.2.
+	// In the case when Both PostProcessing and ToneMapper are disabled we apply gamma manually. In every other case we apply inverse gamma before applying OCIO.
+	float DisplayGamma = (View.Family->EngineShowFlags.Tonemapper == 0) || (View.Family->EngineShowFlags.PostProcessing == 0) ? DefaultDisplayGamma : DefaultDisplayGamma / EngineDisplayGamma;
+
+	FOpenColorIORendering::AddPass_RenderThread(
+		GraphBuilder,
+		View,
+		View.GetFeatureLevel(),
+		Input,
+		Output,
+		InPassInfo,
+		DisplayGamma,
+		TransformAlpha
+	);
+}
+
+void FOpenColorIORendering::PrepareView(FSceneViewFamily& InViewFamily, FSceneView& InView)
+{
+	//Force ToneCurve to be off while we'are alive to make sure the input color space is the working space : srgb linear
+	InViewFamily.EngineShowFlags.SetToneCurve(false);
+	// This flags sets tonampper to output to ETonemapperOutputDevice::LinearNoToneCurve
+	InViewFamily.SceneCaptureSource = SCS_FinalColorHDR;
+
+	InView.FinalPostProcessSettings.bOverride_ToneCurveAmount = 1;
+	InView.FinalPostProcessSettings.ToneCurveAmount = 0.0;
+}
+
 // static
 FOpenColorIORenderPassResources FOpenColorIORendering::GetRenderPassResources(const FOpenColorIOColorConversionSettings& InSettings, ERHIFeatureLevel::Type InFeatureLevel)
 {
diff --git a/Engine/Plugins/Compositing/OpenColorIO/Source/OpenColorIO/Public/OpenColorIODisplayExtension.h b/Engine/Plugins/Compositing/OpenColorIO/Source/OpenColorIO/Public/OpenColorIODisplayExtension.h
--- a/Engine/Plugins/Compositing/OpenColorIO/Source/OpenColorIO/Public/OpenColorIODisplayExtension.h
+++ b/Engine/Plugins/Compositing/OpenColorIO/Source/OpenColorIO/Public/OpenColorIODisplayExtension.h
@@ -58,7 +58,6 @@ public:
 	/** Returns the current display configuration to be updated */
 	FOpenColorIODisplayConfiguration& GetDisplayConfiguration() { return DisplayConfiguration; }
 
-	static float DefaultDisplayGamma;
 private:
 
 	/** Cached pass resources required to apply conversion for render thread */
diff --git a/Engine/Plugins/Compositing/OpenColorIO/Source/OpenColorIO/Public/OpenColorIORendering.h b/Engine/Plugins/Compositing/OpenColorIO/Source/OpenColorIO/Public/OpenColorIORendering.h
--- a/Engine/Plugins/Compositing/OpenColorIO/Source/OpenColorIO/Public/OpenColorIORendering.h
+++ b/Engine/Plugins/Compositing/OpenColorIO/Source/OpenColorIO/Public/OpenColorIORendering.h
@@ -13,6 +13,7 @@ struct FScreenPassTexture;
 struct FScreenPassViewInfo;
 class FTextureResource;
 class FSceneView;
+class FSceneViewFamily;
 class UTexture;
 class UTextureRenderTarget2D;
 class UWorld;
@@ -72,7 +73,8 @@ public:
 	 * Applies the color transform RDG pass with the provided resources.
 	 *
 	 * @param GraphBuilder Render graph builder
-	 * @param View Scene view with additional information
+	 * @param ViewInfo Scene view with additional information
+	 * @param FeatureLevel - Shader model.
 	 * @param Input Input color texture
 	 * @param Output Destination render target
 	 * @param InPassInfo OpenColorIO shader and texture resources
@@ -88,4 +90,30 @@ public:
 		const FOpenColorIORenderPassResources& InPassInfo,
 		float InGamma,
 		EOpenColorIOTransformAlpha TransformAlpha = EOpenColorIOTransformAlpha::None);
+
+	/**
+	 * Similar to the above, except gamma and feature level are handled by the function itself based on provided View.
+	 *
+	 * @param GraphBuilder Render graph builder
+	 * @param View Scene view with additional information
+	 * @param Input Input color texture
+	 * @param Output Destination render target
+	 * @param InPassInfo OpenColorIO shader and texture resources
+	 * @param TransformAlpha Whether to unpremult/invert before applying the color transform
+	 */
+	static void AddPass_RenderThread(
+		FRDGBuilder& GraphBuilder,
+		const FSceneView& View,
+		const FScreenPassTexture& Input,
+		const FScreenPassRenderTarget& Output,
+		const FOpenColorIORenderPassResources& InPassInfo,
+		EOpenColorIOTransformAlpha TransformAlpha = EOpenColorIOTransformAlpha::None);
+
+	/** 
+	* This function sets up provided view and related family to properly support OCIO. 
+	*/
+	static void PrepareView(FSceneViewFamily& InViewFamily, FSceneView& InView);
+	
+	/** */
+	static float DefaultDisplayGamma;
 };
\ No newline at end of file
diff --git a/Engine/Plugins/MovieScene/MovieRenderPipeline/Source/MovieRenderPipelineCore/Private/Graph/Nodes/MovieGraphDeferredRenderPassNode.cpp b/Engine/Plugins/MovieScene/MovieRenderPipeline/Source/MovieRenderPipelineCore/Private/Graph/Nodes/MovieGraphDeferredRenderPassNode.cpp
--- a/Engine/Plugins/MovieScene/MovieRenderPipeline/Source/MovieRenderPipelineCore/Private/Graph/Nodes/MovieGraphDeferredRenderPassNode.cpp
+++ b/Engine/Plugins/MovieScene/MovieRenderPipeline/Source/MovieRenderPipelineCore/Private/Graph/Nodes/MovieGraphDeferredRenderPassNode.cpp
@@ -145,7 +145,7 @@ void UMovieGraphDeferredRenderPassNode::FMovieGraphDeferredRenderPass::Render(co
 	// OCIO: Since this is a manually created Render target we don't need Gamma to be applied.
 	// We use this render target to render to via a display extension that utilizes Display Gamma
 	// which has a default value of 2.2 (DefaultDisplayGamma), therefore we need to set Gamma on this render target to 2.2 to cancel out any unwanted effects.
-	RenderTargetInitParams.TargetGamma = FOpenColorIODisplayExtension::DefaultDisplayGamma;
+	RenderTargetInitParams.TargetGamma = FOpenColorIORendering::DefaultDisplayGamma;
 	RenderTargetInitParams.PixelFormat = EPixelFormat::PF_FloatRGBA;
 
 	UTextureRenderTarget2D* RenderTarget = Renderer->GetOrCreateViewRenderTarget(RenderTargetInitParams);
diff --git a/Engine/Plugins/MovieScene/MovieRenderPipeline/Source/MovieRenderPipelineRenderPasses/Private/MoviePipelineImagePassBase.cpp b/Engine/Plugins/MovieScene/MovieRenderPipeline/Source/MovieRenderPipelineRenderPasses/Private/MoviePipelineImagePassBase.cpp
--- a/Engine/Plugins/MovieScene/MovieRenderPipeline/Source/MovieRenderPipelineRenderPasses/Private/MoviePipelineImagePassBase.cpp
+++ b/Engine/Plugins/MovieScene/MovieRenderPipeline/Source/MovieRenderPipelineRenderPasses/Private/MoviePipelineImagePassBase.cpp
@@ -149,7 +149,7 @@ TWeakObjectPtr<UTextureRenderTarget2D> UMoviePipelineImagePassBase::CreateViewRe
 	// OCIO: Since this is a manually created Render target we don't need Gamma to be applied.
 	// We use this render target to render to via a display extension that utilizes Display Gamma
 	// which has a default value of 2.2 (DefaultDisplayGamma), therefore we need to set Gamma on this render target to 2.2 to cancel out any unwanted effects.
-	NewTarget->TargetGamma = FOpenColorIODisplayExtension::DefaultDisplayGamma;
+	NewTarget->TargetGamma = FOpenColorIORendering::DefaultDisplayGamma;
 
 	// Initialize to the tile size (not final size) and use a 16 bit back buffer to avoid precision issues when accumulating later
 	NewTarget->InitCustomFormat(InSize.X, InSize.Y, EPixelFormat::PF_FloatRGBA, false);
diff --git a/Engine/Plugins/Runtime/nDisplay/Source/DisplayCluster/Private/Render/Viewport/DisplayClusterViewport_OpenColorIO.cpp b/Engine/Plugins/Runtime/nDisplay/Source/DisplayCluster/Private/Render/Viewport/DisplayClusterViewport_OpenColorIO.cpp
--- a/Engine/Plugins/Runtime/nDisplay/Source/DisplayCluster/Private/Render/Viewport/DisplayClusterViewport_OpenColorIO.cpp
+++ b/Engine/Plugins/Runtime/nDisplay/Source/DisplayCluster/Private/Render/Viewport/DisplayClusterViewport_OpenColorIO.cpp
@@ -33,15 +33,7 @@ void FDisplayClusterViewport_OpenColorIO::SetupSceneView(FSceneViewFamily& InOut
 {
 	if (bShaderResourceValid)
 	{
-		// Code copied from FOpenColorIODisplayExtension::SetupView():
-
-		//Force ToneCurve to be off while we'are alive to make sure the input color space is the working space : srgb linear
-		InOutViewFamily.EngineShowFlags.SetToneCurve(false);
-		// This flags sets tonampper to output to ETonemapperOutputDevice::LinearNoToneCurve
-		InOutViewFamily.SceneCaptureSource = SCS_FinalColorHDR;
-
-		InOutView.FinalPostProcessSettings.bOverride_ToneCurveAmount = 1;
-		InOutView.FinalPostProcessSettings.ToneCurveAmount = 0.0;
+		FOpenColorIORendering::PrepareView(InOutViewFamily, InOutView);
 	}
 }
 
@@ -81,7 +73,7 @@ bool FDisplayClusterViewport_OpenColorIO::IsConversionSettingsEqual(const FOpenC
 
 float FDisplayClusterViewport_OpenColorIO::GetDisplayGamma(const FDisplayClusterViewport_Context& InViewportContext) const
 {
-	const float DefaultDisplayGamma = FOpenColorIODisplayExtension::DefaultDisplayGamma;
+	const float DefaultDisplayGamma = FOpenColorIORendering::DefaultDisplayGamma;
 	const float EngineDisplayGamma = InViewportContext.RenderThreadData.EngineDisplayGamma;
 
 	// There is a special case where post processing and tonemapper are disabled. In this case tonemapper applies a static display Inverse of Gamma which defaults to 2.2.
