OCIO: Avoid adding the matrix transform when the source color space equals the working one.

#rnx

diff --git a/Engine/Plugins/Compositing/OpenColorIO/Source/OpenColorIOWrapper/Private/OpenColorIOWrapper.cpp b/Engine/Plugins/Compositing/OpenColorIO/Source/OpenColorIOWrapper/Private/OpenColorIOWrapper.cpp
--- a/Engine/Plugins/Compositing/OpenColorIO/Source/OpenColorIOWrapper/Private/OpenColorIOWrapper.cpp
+++ b/Engine/Plugins/Compositing/OpenColorIO/Source/OpenColorIOWrapper/Private/OpenColorIOWrapper.cpp
@@ -827,29 +827,34 @@ FOpenColorIOWrapperProcessor FOpenColorIOWrapperProcessor::CreateTransformToWork
 			break;
 		}
 
-		const EChromaticAdaptationMethod ChromaticAdapation = static_cast<EChromaticAdaptationMethod>(InColorSettings.ChromaticAdaptationMethod);
+		TOptional<FColorSpace> SourceColorSpace;
 
 		if (InColorSettings.ColorSpaceOverride.IsSet())
 		{
-			const FColorSpace SourceColorSpace = FColorSpace
-			(
-				InColorSettings.ColorSpaceOverride.GetValue()[0],
-				InColorSettings.ColorSpaceOverride.GetValue()[1],
-				InColorSettings.ColorSpaceOverride.GetValue()[2],
-				InColorSettings.ColorSpaceOverride.GetValue()[3]
+			const TStaticArray<FVector2d, 4>& SourceChromaticities = InColorSettings.ColorSpaceOverride.GetValue();
+			SourceColorSpace = FColorSpace(
+				SourceChromaticities[0],
+				SourceChromaticities[1],
+				SourceChromaticities[2],
+				SourceChromaticities[3]
 			);
-			MatrixTransformRcPtr MatrixTransform = MatrixTransform::Create();
-			const FMatrix44d ToWorkingMat = Transpose<double>(FColorSpaceTransform(SourceColorSpace, FColorSpace::GetWorking(), ChromaticAdapation));
-			MatrixTransform->setMatrix(&ToWorkingMat.M[0][0]);
-			TransformToWCS->appendTransform(MatrixTransform);
 		}
 		else if (InColorSettings.ColorSpace != EColorSpace::None)
 		{
-			const FColorSpace SourceColorSpace = FColorSpace(static_cast<EColorSpace>(InColorSettings.ColorSpace));
-			const FMatrix44d ToWorkingMat = Transpose<double>(FColorSpaceTransform(SourceColorSpace, FColorSpace::GetWorking(), ChromaticAdapation));
-			MatrixTransformRcPtr MatrixTransform = MatrixTransform::Create();
-			MatrixTransform->setMatrix(&ToWorkingMat.M[0][0]);
-			TransformToWCS->appendTransform(MatrixTransform);
+			SourceColorSpace = FColorSpace(InColorSettings.ColorSpace);
+		}
+
+		if (SourceColorSpace.IsSet())
+		{
+			if (!SourceColorSpace.GetValue().Equals(FColorSpace::GetWorking()))
+			{
+				const EChromaticAdaptationMethod ChromaticAdapation = static_cast<EChromaticAdaptationMethod>(InColorSettings.ChromaticAdaptationMethod);
+				const FMatrix44d ToWorkingMat = Transpose<double>(FColorSpaceTransform(SourceColorSpace.GetValue(), FColorSpace::GetWorking(), ChromaticAdapation));
+				MatrixTransformRcPtr MatrixTransform = MatrixTransform::Create();
+				MatrixTransform->setMatrix(&ToWorkingMat.M[0][0]);
+				
+				TransformToWCS->appendTransform(MatrixTransform);
+			}
 		}
 
 		ParentTransform->setTransform(TransformToWCS, TRANSFORM_DIR_FORWARD);
diff --git a/Engine/Source/Runtime/ColorManagement/Private/ColorSpace.cpp b/Engine/Source/Runtime/ColorManagement/Private/ColorSpace.cpp
--- a/Engine/Source/Runtime/ColorManagement/Private/ColorSpace.cpp
+++ b/Engine/Source/Runtime/ColorManagement/Private/ColorSpace.cpp
@@ -201,10 +201,15 @@ FMatrix44d FColorSpace::CalcRgbToXYZ() const
 
 bool FColorSpace::Equals(const FColorSpace& ColorSpace, double Tolerance) const
 {
-	return	Chromaticities[0].Equals(ColorSpace.Chromaticities[0], Tolerance) &&
-		Chromaticities[1].Equals(ColorSpace.Chromaticities[1], Tolerance) &&
-		Chromaticities[2].Equals(ColorSpace.Chromaticities[2], Tolerance) &&
-		Chromaticities[3].Equals(ColorSpace.Chromaticities[3], Tolerance);
+	return Equals(ColorSpace.Chromaticities, Tolerance);
+}
+
+bool FColorSpace::Equals(const TStaticArray<FVector2d, 4>& InChromaticities, double Tolerance) const
+{
+	return Chromaticities[0].Equals(InChromaticities[0], Tolerance) &&
+		Chromaticities[1].Equals(InChromaticities[1], Tolerance) &&
+		Chromaticities[2].Equals(InChromaticities[2], Tolerance) &&
+		Chromaticities[3].Equals(InChromaticities[3], Tolerance);
 }
 
 bool FColorSpace::IsSRGB() const
diff --git a/Engine/Source/Runtime/ColorManagement/Public/ColorSpace.h b/Engine/Source/Runtime/ColorManagement/Public/ColorSpace.h
--- a/Engine/Source/Runtime/ColorManagement/Public/ColorSpace.h
+++ b/Engine/Source/Runtime/ColorManagement/Public/ColorSpace.h
@@ -250,7 +250,7 @@ public:
 	 * Check against another color space for equality.
 	 *
 	 * @param ColorSpace The vector to check against.
-	 * @return true if the vectors are equal, false otherwise.
+	 * @return true if the color spaces are equal, false otherwise.
 	 */
 	inline bool operator==(const FColorSpace& ColorSpace) const
 	{
@@ -260,8 +260,8 @@ public:
 	/**
 	 * Check against another color space for inequality.
 	 *
-	 * @param ColorSpace The vector to check against.
-	 * @return true if the vectors are not equal, false otherwise.
+	 * @param ColorSpace The color space to check against.
+	 * @return true if the color spaces are not equal, false otherwise.
 	 */
 	inline bool operator!=(const FColorSpace& ColorSpace) const
 	{
@@ -271,12 +271,22 @@ public:
 	/**
 	 * Check against another color space for equality, within specified error limits.
 	 *
-	 * @param ColorSpace The vector to check against.
+	 * @param ColorSpace The color space to check against.
 	 * @param Tolerance Error tolerance.
-	 * @return true if the vectors are equal within tolerance limits, false otherwise.
+	 * @return true if the color spaces are equal within tolerance limits, false otherwise.
 	 */
 	bool Equals(const FColorSpace& ColorSpace, double Tolerance = 1.e-7) const;
 
+
+	/**
+	 * Check against color space chromaticities for equality, within specified error limits.
+	 *
+	 * @param InChromaticities The color space chromaticities to check against.
+	 * @param Tolerance Error tolerance.
+	 * @return true if the color spaces are equal within tolerance limits, false otherwise.
+	 */
+	bool Equals(const TStaticArray<FVector2d, 4>& InChromaticities, double Tolerance = 1.e-7) const;
+
 	/**
 	 * Convenience function to verify if the color space matches the engine's default sRGB chromaticities.
 	 *
