OCIO: Parallelized (single) image processing.

Also removed heap allocations for image desc structs, and restructured the code around the parallel inner loop.

#rnx

diff --git a/Engine/Plugins/Compositing/OpenColorIO/Source/OpenColorIOWrapper/Private/OpenColorIOWrapper.cpp b/Engine/Plugins/Compositing/OpenColorIO/Source/OpenColorIOWrapper/Private/OpenColorIOWrapper.cpp
--- a/Engine/Plugins/Compositing/OpenColorIO/Source/OpenColorIOWrapper/Private/OpenColorIOWrapper.cpp
+++ b/Engine/Plugins/Compositing/OpenColorIO/Source/OpenColorIOWrapper/Private/OpenColorIOWrapper.cpp
@@ -9,6 +9,7 @@
 #include "ColorSpace.h"
 #include "Engine/TextureDefines.h"
 #include "ImageCore.h"
+#include "ImageParallelFor.h"
 
 #if WITH_OCIO
 THIRD_PARTY_INCLUDES_START
@@ -99,7 +100,22 @@ struct FOpenColorIOGPUProcessorPimpl
 
 #if WITH_OCIO
 namespace {
-	TUniquePtr<OCIO_NAMESPACE::PackedImageDesc> GetImageDesc(const FImageView& InImage)
+
+	bool IsImageFormatSupported(const FImageView& InImage)
+	{
+		switch (InImage.Format)
+		{
+		case ERawImageFormat::BGRA8:
+		case ERawImageFormat::RGBA16:
+		case ERawImageFormat::RGBA16F:
+		case ERawImageFormat::RGBA32F:
+			return true;
+		default:
+			return false;
+		}
+	}
+
+	OCIO_NAMESPACE::PackedImageDesc GetImageDesc(const FImageView& InImage)
 	{
 		OCIO_NAMESPACE::ChannelOrdering Ordering;
 		OCIO_NAMESPACE::BitDepth BitDepth;
@@ -124,11 +140,11 @@ namespace {
 			break;
 		default:
 
-			UE_LOG(LogOpenColorIOWrapper, Log, TEXT("Unsupported texture format."));
-			return nullptr;
+			// All other cases should have been eliminated by calling IsImageFormatSupported() previously.
+			checkNoEntry();
 		}
 
-		return MakeUnique<OCIO_NAMESPACE::PackedImageDesc>(
+		return OCIO_NAMESPACE::PackedImageDesc(
 			InImage.RawData,
 			static_cast<long>(InImage.GetWidth()),
 			static_cast<long>(InImage.GetHeight()),
@@ -923,7 +939,7 @@ bool FOpenColorIOWrapperProcessor::TransformColor(FLinearColor& InOutColor) cons
 
 			return true;
 		}
-	OCIO_EXCEPTION_HANDLING_CATCH(Log, TEXT("Failed to transform color. Error message: %s"));
+	OCIO_EXCEPTION_HANDLING_CATCH(Error, TEXT("Failed to transform color. Error message: %s"));
 #endif // WITH_OCIO
 
 	return false;
@@ -938,49 +954,74 @@ bool FOpenColorIOWrapperProcessor::TransformImage(const FImageView& InOutImage)
 
 		if (IsValid())
 		{
-			TUniquePtr<PackedImageDesc> ImageDesc = GetImageDesc(InOutImage);
-			if (ImageDesc)
+			if (IsImageFormatSupported(InOutImage))
 			{
-				const TUniquePtr<ANSICHAR[]> AnsiWorkingColorSpaceName = OpenColorIOWrapper::MakeAnsiString(OpenColorIOWrapper::GetWorkingColorSpaceName());
-				ConstConfigRcPtr	InterchangeConfig = IOpenColorIOWrapperModule::Get().GetEngineBuiltInConfig().Pimpl->Config;
-				ConstConfigRcPtr Config = OwnerConfig->Pimpl->Config;
+				const BitDepth BitDepth = GetImageDesc(InOutImage).getBitDepth();
 
-				BitDepth BitDepth = ImageDesc->getBitDepth();
+				// Primary transform processor
+				ConstCPUProcessorRcPtr CPUProcessor = Pimpl->Processor->getOptimizedCPUProcessor(BitDepth, BitDepth, OPTIMIZATION_DEFAULT);
 
-				// Conditionally apply a conversion from the working color space to interchange space
-				if (WorkingColorSpaceTransformType == EOpenColorIOWorkingColorSpaceTransform::Source)
+				// Note: This special "interchange to working color space" logic should be removed
+				// once these features are supported by the library.
+				ConstProcessorRcPtr InterchangeProcessor = nullptr;
+				ConstCPUProcessorRcPtr InterchangeCPUProcessor = nullptr;
 				{
-					ConstProcessorRcPtr	InterchangeProcessor = InterchangeConfig->GetProcessorFromConfigs(
-						InterchangeConfig,
-						AnsiWorkingColorSpaceName.Get(),
-						Config,
-						Config->getCanonicalName(OpenColorIOWrapper::GetInterchangeName()));
-
-					ConstCPUProcessorRcPtr InterchangeCPUProcessor = InterchangeProcessor->getOptimizedCPUProcessor(BitDepth, BitDepth, OPTIMIZATION_DEFAULT);
-					InterchangeCPUProcessor->apply(*ImageDesc);
+					const TUniquePtr<ANSICHAR[]> AnsiWorkingColorSpaceName = OpenColorIOWrapper::MakeAnsiString(OpenColorIOWrapper::GetWorkingColorSpaceName());
+					ConstConfigRcPtr InterchangeConfig = IOpenColorIOWrapperModule::Get().GetEngineBuiltInConfig().Pimpl->Config;
+					ConstConfigRcPtr Config = OwnerConfig->Pimpl->Config;
+
+					if (WorkingColorSpaceTransformType == EOpenColorIOWorkingColorSpaceTransform::Source)
+					{
+						InterchangeProcessor = InterchangeConfig->GetProcessorFromConfigs(
+							InterchangeConfig,
+							AnsiWorkingColorSpaceName.Get(),
+							Config,
+							Config->getCanonicalName(OpenColorIOWrapper::GetInterchangeName()));
+					}
+					else if (WorkingColorSpaceTransformType == EOpenColorIOWorkingColorSpaceTransform::Destination)
+					{
+						InterchangeProcessor = InterchangeConfig->GetProcessorFromConfigs(
+							Config,
+							Config->getCanonicalName(OpenColorIOWrapper::GetInterchangeName()),
+							InterchangeConfig,
+							AnsiWorkingColorSpaceName.Get());
+					}
+
+					if (InterchangeProcessor)
+					{
+						InterchangeCPUProcessor = InterchangeProcessor->getOptimizedCPUProcessor(BitDepth, BitDepth, OPTIMIZATION_DEFAULT);
+					}
 				}
 
-				// Apply the main color transformation
-				ConstCPUProcessorRcPtr CPUProcessor = Pimpl->Processor->getOptimizedCPUProcessor(BitDepth, BitDepth, OPTIMIZATION_DEFAULT);
-				CPUProcessor->apply(*ImageDesc);
+				// Apply parallelized color transformation
+				FImageCore::ImageParallelFor(TEXT("FOpenColorIOWrapperProcessor.TransformImage.PF"), InOutImage, [&](FImageView& ImagePart)
+					{
+						OCIO_NAMESPACE::PackedImageDesc ImagePartDesc = GetImageDesc(ImagePart);
 
-				// Conditionally apply a conversion from the interchange space to the working color space
-				if (WorkingColorSpaceTransformType == EOpenColorIOWorkingColorSpaceTransform::Destination)
-				{
-					ConstProcessorRcPtr	InterchangeProcessor = InterchangeConfig->GetProcessorFromConfigs(
-						Config,
-						Config->getCanonicalName(OpenColorIOWrapper::GetInterchangeName()),
-						InterchangeConfig,
-						AnsiWorkingColorSpaceName.Get());
+						// Conditionally apply a conversion from the working color space to interchange space
+						if (WorkingColorSpaceTransformType == EOpenColorIOWorkingColorSpaceTransform::Source)
+						{
+							InterchangeCPUProcessor->apply(ImagePartDesc);
+						}
 
-					ConstCPUProcessorRcPtr InterchangeCPUProcessor = InterchangeProcessor->getOptimizedCPUProcessor(BitDepth, BitDepth, OPTIMIZATION_DEFAULT);
-					InterchangeCPUProcessor->apply(*ImageDesc);
-				}
+						// Apply the main color transformation
+						CPUProcessor->apply(ImagePartDesc);
+
+						// Conditionally apply a conversion from the interchange space to the working color space
+						if (WorkingColorSpaceTransformType == EOpenColorIOWorkingColorSpaceTransform::Destination)
+						{
+							InterchangeCPUProcessor->apply(ImagePartDesc);
+						}
+					});
 
 				return true;
 			}
+			else
+			{
+				UE_LOG(LogOpenColorIOWrapper, Warning, TEXT("Unsupported texture format."));
+			}
 		}
-	OCIO_EXCEPTION_HANDLING_CATCH(Log, TEXT("Failed to transform image. Error message: %s"));
+	OCIO_EXCEPTION_HANDLING_CATCH(Error, TEXT("Failed to transform image. Error message: %s"));
 #endif // WITH_OCIO
 
 	return false;
@@ -995,16 +1036,17 @@ bool FOpenColorIOWrapperProcessor::TransformImage(const FImageView& SrcImage, co
 
 		if (IsValid())
 		{
-			TUniquePtr<PackedImageDesc> SrcImageDesc = GetImageDesc(SrcImage);
-			TUniquePtr<PackedImageDesc> DestImageDesc = GetImageDesc(DestImage);
-			if (SrcImageDesc && DestImageDesc)
+			if (IsImageFormatSupported(SrcImage) && IsImageFormatSupported(DestImage))
 			{
+				PackedImageDesc SrcImageDesc = GetImageDesc(SrcImage);
+				PackedImageDesc DestImageDesc = GetImageDesc(DestImage);
+
 				const TUniquePtr<ANSICHAR[]> AnsiWorkingColorSpaceName = OpenColorIOWrapper::MakeAnsiString(OpenColorIOWrapper::GetWorkingColorSpaceName());
 				ConstConfigRcPtr	InterchangeConfig = IOpenColorIOWrapperModule::Get().GetEngineBuiltInConfig().Pimpl->Config;
 				ConstConfigRcPtr	Config = OwnerConfig->Pimpl->Config;
 
-				BitDepth SrcBitDepth = SrcImageDesc->getBitDepth();
-				BitDepth DestBitDepth = DestImageDesc->getBitDepth();
+				BitDepth SrcBitDepth = SrcImageDesc.getBitDepth();
+				BitDepth DestBitDepth = DestImageDesc.getBitDepth();
 
 				// Conditionally apply a conversion from the working color space to interchange space
 				if (WorkingColorSpaceTransformType == EOpenColorIOWorkingColorSpaceTransform::Source)
@@ -1016,13 +1058,13 @@ bool FOpenColorIOWrapperProcessor::TransformImage(const FImageView& SrcImage, co
 						Config->getCanonicalName(OpenColorIOWrapper::GetInterchangeName()));
 
 					ConstCPUProcessorRcPtr InterchangeCPUProcessor = InterchangeProcessor->getOptimizedCPUProcessor(SrcBitDepth, SrcBitDepth, OPTIMIZATION_DEFAULT);
-					InterchangeCPUProcessor->apply(*SrcImageDesc);
+					InterchangeCPUProcessor->apply(SrcImageDesc);
 				}
 
 
 				// Apply the main color transformation
 				ConstCPUProcessorRcPtr CPUProcessor = Pimpl->Processor->getOptimizedCPUProcessor(SrcBitDepth, DestBitDepth, OPTIMIZATION_DEFAULT);
-				CPUProcessor->apply(*SrcImageDesc, *DestImageDesc);
+				CPUProcessor->apply(SrcImageDesc, DestImageDesc);
 
 				// Conditionally apply a conversion from the interchange space to the working color space
 				if (WorkingColorSpaceTransformType == EOpenColorIOWorkingColorSpaceTransform::Destination)
@@ -1034,13 +1076,17 @@ bool FOpenColorIOWrapperProcessor::TransformImage(const FImageView& SrcImage, co
 						AnsiWorkingColorSpaceName.Get());
 
 					ConstCPUProcessorRcPtr InterchangeCPUProcessor = InterchangeProcessor->getOptimizedCPUProcessor(DestBitDepth, DestBitDepth, OPTIMIZATION_DEFAULT);
-					InterchangeCPUProcessor->apply(*DestImageDesc);
+					InterchangeCPUProcessor->apply(DestImageDesc);
 				}
 
 				return true;
 			}
+			else
+			{
+				UE_LOG(LogOpenColorIOWrapper, Warning, TEXT("Unsupported texture format(s)."));
+			}
 		}
-	OCIO_EXCEPTION_HANDLING_CATCH(Log, TEXT("Failed to transform image. Error message: %s"));
+	OCIO_EXCEPTION_HANDLING_CATCH(Error, TEXT("Failed to transform image. Error message: %s"));
 #endif // WITH_OCIO
 
 	return false;
diff --git a/Engine/Plugins/Compositing/OpenColorIO/Source/OpenColorIOWrapper/Public/OpenColorIOWrapper.h b/Engine/Plugins/Compositing/OpenColorIO/Source/OpenColorIOWrapper/Public/OpenColorIOWrapper.h
--- a/Engine/Plugins/Compositing/OpenColorIO/Source/OpenColorIOWrapper/Public/OpenColorIOWrapper.h
+++ b/Engine/Plugins/Compositing/OpenColorIO/Source/OpenColorIOWrapper/Public/OpenColorIOWrapper.h
@@ -227,7 +227,11 @@ public:
 	/** Apply the CPU color transform in-place to the specified image. */
 	OPENCOLORIOWRAPPER_API bool TransformImage(const FImageView& InOutImage) const;
 
-	/** Apply the CPU color transform from the source image to the destination image. (The destination FImageView is const but what it points at is not.) */
+	/** Apply the CPU color transform from the source image to the destination image.
+	 * (The destination FImageView is const but what it points at is not.)
+	 * 
+	 * Note: This function is currently less optimal, as it is not parallelized.
+	 */
 	OPENCOLORIOWRAPPER_API bool TransformImage(const FImageView& SrcImage, const FImageView& DestImage) const;
 
 private:
