OCIO: Fix uninitialized variable warning.

diff --git a/Engine/Plugins/Compositing/OpenColorIO/Source/OpenColorIOWrapper/Private/OpenColorIOWrapper.cpp b/Engine/Plugins/Compositing/OpenColorIO/Source/OpenColorIOWrapper/Private/OpenColorIOWrapper.cpp
--- a/Engine/Plugins/Compositing/OpenColorIO/Source/OpenColorIOWrapper/Private/OpenColorIOWrapper.cpp
+++ b/Engine/Plugins/Compositing/OpenColorIO/Source/OpenColorIOWrapper/Private/OpenColorIOWrapper.cpp
@@ -117,8 +117,8 @@ namespace {
 
 	OCIO_NAMESPACE::PackedImageDesc GetImageDesc(const FImageView& InImage)
 	{
-		OCIO_NAMESPACE::ChannelOrdering Ordering;
-		OCIO_NAMESPACE::BitDepth BitDepth;
+		OCIO_NAMESPACE::ChannelOrdering Ordering = OCIO_NAMESPACE::ChannelOrdering::CHANNEL_ORDERING_RGBA;
+		OCIO_NAMESPACE::BitDepth BitDepth = OCIO_NAMESPACE::BitDepth::BIT_DEPTH_UNKNOWN;
 
 		switch (InImage.Format)
 		{
