OCIO: Fix to ensure the OCIO is also active when a highres screenshot is being taken.

diff --git a/Engine/Plugins/Compositing/OpenColorIO/Source/OpenColorIO/Private/OpenColorIODisplayExtensionWrapper.cpp b/Engine/Plugins/Compositing/OpenColorIO/Source/OpenColorIO/Private/OpenColorIODisplayExtensionWrapper.cpp
--- a/Engine/Plugins/Compositing/OpenColorIO/Source/OpenColorIO/Private/OpenColorIODisplayExtensionWrapper.cpp
+++ b/Engine/Plugins/Compositing/OpenColorIO/Source/OpenColorIO/Private/OpenColorIODisplayExtensionWrapper.cpp
@@ -124,6 +124,12 @@ UOpenColorIODisplayExtensionWrapper* UOpenColorIODisplayExtensionWrapper::Create
 			}
 		}
 
+		// Activate the SVE if a high-res screenshot is being taken
+		if (GIsHighResScreenshot)
+		{
+			return TOptional<bool>(bIsEnabled);
+		}
+
 		// If our viewport did not meet any of the criteria to activate the SVE, emit no opinion.
 		return TOptional<bool>();
 	};
