OCIO: Prevent slate crash on invalid configuration settings.

Also simplified editor blueprint function signatures.

#rnx

diff --git a/Engine/Plugins/Compositing/OpenColorIO/Source/OpenColorIO/Private/OpenColorIOColorSpace.cpp b/Engine/Plugins/Compositing/OpenColorIO/Source/OpenColorIO/Private/OpenColorIOColorSpace.cpp
--- a/Engine/Plugins/Compositing/OpenColorIO/Source/OpenColorIO/Private/OpenColorIOColorSpace.cpp
+++ b/Engine/Plugins/Compositing/OpenColorIO/Source/OpenColorIO/Private/OpenColorIOColorSpace.cpp
@@ -114,7 +114,7 @@ FOpenColorIOColorConversionSettings::FOpenColorIOColorConversionSettings()
 
 FString FOpenColorIOColorConversionSettings::ToString() const
 {
-	if (ConfigurationSource)
+	if (::IsValid(ConfigurationSource))
 	{
 		if (IsDisplayView())
 		{
@@ -139,7 +139,7 @@ FString FOpenColorIOColorConversionSettings::ToString() const
 
 bool FOpenColorIOColorConversionSettings::IsValid() const
 {
-	if (ConfigurationSource)
+	if (::IsValid(ConfigurationSource))
 	{
 		if (IsDisplayView())
 		{
@@ -171,7 +171,7 @@ bool FOpenColorIOColorConversionSettings::IsValid() const
 
 void FOpenColorIOColorConversionSettings::ValidateColorSpaces()
 {
-	if (ConfigurationSource)
+	if (::IsValid(ConfigurationSource))
 	{
 		if (!ConfigurationSource->HasDesiredColorSpace(SourceColorSpace))
 		{
diff --git a/Engine/Plugins/Compositing/OpenColorIO/Source/OpenColorIOEditor/Private/OpenColorIOColorSpaceConversionCustomization.cpp b/Engine/Plugins/Compositing/OpenColorIO/Source/OpenColorIOEditor/Private/OpenColorIOColorSpaceConversionCustomization.cpp
--- a/Engine/Plugins/Compositing/OpenColorIO/Source/OpenColorIOEditor/Private/OpenColorIOColorSpaceConversionCustomization.cpp
+++ b/Engine/Plugins/Compositing/OpenColorIO/Source/OpenColorIOEditor/Private/OpenColorIOColorSpaceConversionCustomization.cpp
@@ -20,49 +20,48 @@
 
 void FOpenColorIOColorConversionSettingsCustomization::CustomizeHeader(TSharedRef<IPropertyHandle> InPropertyHandle, FDetailWidgetRow& HeaderRow, IPropertyTypeCustomizationUtils& CustomizationUtils)
 {
-	if (InPropertyHandle->GetNumPerObjectValues() == 1 && InPropertyHandle->IsValidHandle())
-	{
-		void* StructData = nullptr;
-		if (InPropertyHandle->GetValueData(StructData) == FPropertyAccess::Success)
-		{
-			ColorSpaceConversion = reinterpret_cast<FOpenColorIOColorConversionSettings*>(StructData);
-			check(ColorSpaceConversion);
-
-			ColorSpaceConversion->ValidateColorSpaces();
+	ColorSpaceSettingsProperty = InPropertyHandle;
 
-			TSharedPtr<IPropertyUtilities> PropertyUtils = CustomizationUtils.GetPropertyUtilities();
-
-			HeaderRow.NameContent()
-				[
-					InPropertyHandle->CreatePropertyNameWidget()
-				]
-				.ValueContent()
-				.MaxDesiredWidth(512)
+	if (FOpenColorIOColorConversionSettings* ColorSpaceConversion = GetConversionSettings())
+	{
+		ColorSpaceConversion->ValidateColorSpaces();
+
+		TSharedPtr<IPropertyUtilities> PropertyUtils = CustomizationUtils.GetPropertyUtilities();
+
+		HeaderRow.NameContent()
+			[
+				InPropertyHandle->CreatePropertyNameWidget()
+			]
+			.ValueContent()
+			.MaxDesiredWidth(512)
+			[
+				SNew(SHorizontalBox)
+				+ SHorizontalBox::Slot()
+				.FillWidth(1.0f)
+				.VAlign(VAlign_Center)
 				[
-					SNew(SHorizontalBox)
-					+ SHorizontalBox::Slot()
-					.FillWidth(1.0f)
-					.VAlign(VAlign_Center)
-					[
-						SNew(STextBlock)
-						.Text(MakeAttributeLambda([this]
+					SNew(STextBlock)
+					.Text(MakeAttributeLambda([this]
+						{
+							if (FOpenColorIOColorConversionSettings* ColorSpaceConversion = GetConversionSettings())
 							{
 								if (ColorSpaceConversion->IsValid())
 								{
 									return FText::FromString(*ColorSpaceConversion->ToString());
 								}
+							}
 
-								return FText::FromString(TEXT("<Invalid Conversion>"));
+							return FText::FromString(TEXT("<Invalid Conversion>"));
 
-							}))
-					]
-				].IsEnabled(MakeAttributeLambda([=] { return !InPropertyHandle->IsEditConst() && PropertyUtils->IsPropertyEditingEnabled(); }));
-		}
+						}))
+				]
+			].IsEnabled(MakeAttributeLambda([InPropertyHandle, PropertyUtils] { return !InPropertyHandle->IsEditConst() && PropertyUtils->IsPropertyEditingEnabled(); }));
 	}
 }
 
 void FOpenColorIOColorConversionSettingsCustomization::CustomizeChildren(TSharedRef<IPropertyHandle> InStructPropertyHandle, IDetailChildrenBuilder& StructBuilder, IPropertyTypeCustomizationUtils& StructCustomizationUtils)
 {
+	FOpenColorIOColorConversionSettings* ColorSpaceConversion = GetConversionSettings();
 	if (ColorSpaceConversion == nullptr)
 	{
 		return;
@@ -81,6 +80,8 @@ void FOpenColorIOColorConversionSettingsCustomization::CustomizeChildren(TShared
 
 				ChildHandle->SetOnPropertyValueChanged(FSimpleDelegate::CreateLambda([this]
 					{
+						FOpenColorIOColorConversionSettings* ColorSpaceConversion = GetConversionSettings();
+
 						TransformPicker[OCIO_Src]->SetConfiguration(ColorSpaceConversion->ConfigurationSource);
 						TransformPicker[OCIO_Dst]->SetConfiguration(ColorSpaceConversion->ConfigurationSource);
 
@@ -178,7 +179,10 @@ void FOpenColorIOColorConversionSettingsCustomization::OnSourceColorSpaceChanged
 
 	ApplySelectionToConfiguration();
 
-	ColorSpaceConversion->OnConversionSettingsChanged().Broadcast();
+	if (FOpenColorIOColorConversionSettings* ColorSpaceConversion = GetConversionSettings())
+	{
+		ColorSpaceConversion->OnConversionSettingsChanged().Broadcast();
+	}
 }
 
 void FOpenColorIOColorConversionSettingsCustomization::OnDestinationColorSpaceChanged(const FOpenColorIOColorSpace& NewColorSpace, const FOpenColorIODisplayView& NewDisplayView)
@@ -190,24 +194,30 @@ void FOpenColorIOColorConversionSettingsCustomization::OnDestinationColorSpaceCh
 
 	ApplySelectionToConfiguration();
 
-	ColorSpaceConversion->OnConversionSettingsChanged().Broadcast();
+	if (FOpenColorIOColorConversionSettings* ColorSpaceConversion = GetConversionSettings())
+	{
+		ColorSpaceConversion->OnConversionSettingsChanged().Broadcast();
+	}
 }
 
 void FOpenColorIOColorConversionSettingsCustomization::ApplyConfigurationToSelection()
 {
-	if (ColorSpaceConversion->DisplayViewDirection == EOpenColorIOViewTransformDirection::Forward)
+	if (FOpenColorIOColorConversionSettings* ColorSpaceConversion = GetConversionSettings())
 	{
-		TransformSelection[OCIO_Src].ColorSpace = ColorSpaceConversion->SourceColorSpace;
-		TransformSelection[OCIO_Src].DisplayView.Reset();
-		TransformSelection[OCIO_Dst].ColorSpace = ColorSpaceConversion->DestinationColorSpace;
-		TransformSelection[OCIO_Dst].DisplayView = ColorSpaceConversion->DestinationDisplayView;
-	}
-	else
-	{
-		TransformSelection[OCIO_Src].ColorSpace.Reset();
-		TransformSelection[OCIO_Src].DisplayView = ColorSpaceConversion->DestinationDisplayView;
-		TransformSelection[OCIO_Dst].ColorSpace = ColorSpaceConversion->SourceColorSpace;
-		TransformSelection[OCIO_Dst].DisplayView.Reset();
+		if (ColorSpaceConversion->DisplayViewDirection == EOpenColorIOViewTransformDirection::Forward)
+		{
+			TransformSelection[OCIO_Src].ColorSpace = ColorSpaceConversion->SourceColorSpace;
+			TransformSelection[OCIO_Src].DisplayView.Reset();
+			TransformSelection[OCIO_Dst].ColorSpace = ColorSpaceConversion->DestinationColorSpace;
+			TransformSelection[OCIO_Dst].DisplayView = ColorSpaceConversion->DestinationDisplayView;
+		}
+		else
+		{
+			TransformSelection[OCIO_Src].ColorSpace.Reset();
+			TransformSelection[OCIO_Src].DisplayView = ColorSpaceConversion->DestinationDisplayView;
+			TransformSelection[OCIO_Dst].ColorSpace = ColorSpaceConversion->SourceColorSpace;
+			TransformSelection[OCIO_Dst].DisplayView.Reset();
+		}
 	}
 }
 
@@ -231,19 +241,22 @@ void FOpenColorIOColorConversionSettingsCustomization::ApplySelectionToConfigura
 		Property->NotifyPreChange();
 	}
 
-	if (TransformSelection[OCIO_Src].DisplayView.IsValid())
+	if (FOpenColorIOColorConversionSettings* ColorSpaceConversion = GetConversionSettings())
 	{
-		ColorSpaceConversion->SourceColorSpace = TransformSelection[OCIO_Dst].ColorSpace;
-		ColorSpaceConversion->DestinationColorSpace.Reset();
-		ColorSpaceConversion->DestinationDisplayView = TransformSelection[OCIO_Src].DisplayView;
-		ColorSpaceConversion->DisplayViewDirection = EOpenColorIOViewTransformDirection::Inverse;
-	}
-	else
-	{
-		ColorSpaceConversion->SourceColorSpace = TransformSelection[OCIO_Src].ColorSpace;
-		ColorSpaceConversion->DestinationColorSpace = TransformSelection[OCIO_Dst].ColorSpace;
-		ColorSpaceConversion->DestinationDisplayView = TransformSelection[OCIO_Dst].DisplayView;
-		ColorSpaceConversion->DisplayViewDirection = EOpenColorIOViewTransformDirection::Forward;
+		if (TransformSelection[OCIO_Src].DisplayView.IsValid())
+		{
+			ColorSpaceConversion->SourceColorSpace = TransformSelection[OCIO_Dst].ColorSpace;
+			ColorSpaceConversion->DestinationColorSpace.Reset();
+			ColorSpaceConversion->DestinationDisplayView = TransformSelection[OCIO_Src].DisplayView;
+			ColorSpaceConversion->DisplayViewDirection = EOpenColorIOViewTransformDirection::Inverse;
+		}
+		else
+		{
+			ColorSpaceConversion->SourceColorSpace = TransformSelection[OCIO_Src].ColorSpace;
+			ColorSpaceConversion->DestinationColorSpace = TransformSelection[OCIO_Dst].ColorSpace;
+			ColorSpaceConversion->DestinationDisplayView = TransformSelection[OCIO_Dst].DisplayView;
+			ColorSpaceConversion->DisplayViewDirection = EOpenColorIOViewTransformDirection::Forward;
+		}
 	}
 
 	for (const TSharedPtr<IPropertyHandle>& Property : Properties)
@@ -261,7 +274,24 @@ void FOpenColorIOColorConversionSettingsCustomization::OnConfigurationReset()
 	TransformPicker[OCIO_Src]->SetConfiguration(nullptr);
 	TransformPicker[OCIO_Dst]->SetConfiguration(nullptr);
 
-	ColorSpaceConversion->OnConversionSettingsChanged().Broadcast();
+	if (FOpenColorIOColorConversionSettings* ColorSpaceConversion = GetConversionSettings())
+	{
+		ColorSpaceConversion->OnConversionSettingsChanged().Broadcast();
+	}
+}
+
+FOpenColorIOColorConversionSettings* FOpenColorIOColorConversionSettingsCustomization::GetConversionSettings() const
+{
+	if (ColorSpaceSettingsProperty->IsValidHandle())
+	{
+		void* Data = nullptr;
+		if (ColorSpaceSettingsProperty->GetValueData(Data) == FPropertyAccess::Success)
+		{
+			return static_cast<FOpenColorIOColorConversionSettings*>(Data);
+		}
+	}
+
+	return nullptr;
 }
 
 #undef LOCTEXT_NAMESPACE
diff --git a/Engine/Plugins/Compositing/OpenColorIO/Source/OpenColorIOEditor/Private/OpenColorIOColorSpaceConversionCustomization.h b/Engine/Plugins/Compositing/OpenColorIO/Source/OpenColorIOEditor/Private/OpenColorIOColorSpaceConversionCustomization.h
--- a/Engine/Plugins/Compositing/OpenColorIO/Source/OpenColorIOEditor/Private/OpenColorIOColorSpaceConversionCustomization.h
+++ b/Engine/Plugins/Compositing/OpenColorIO/Source/OpenColorIOEditor/Private/OpenColorIOColorSpaceConversionCustomization.h
@@ -36,6 +36,9 @@ private:
 	/** Callback to reset the configuration in the transform source/destination pickers. */
 	void OnConfigurationReset();
 
+	/** Get the struct value of the color space settings member property.*/
+	struct FOpenColorIOColorConversionSettings* GetConversionSettings() const;
+
 	/** Pointer to the struct SourceColorSpace property handle. */
 	TSharedPtr<IPropertyHandle> SourceColorSpaceProperty;
 	
@@ -53,6 +56,6 @@ private:
 	/** Intermediate selection objects. This is done to allow inverted display-view selections (when enabled in settings). */
 	TStaticArray<FOpenColorIOPickerSelection, 2> TransformSelection;
 
-	/** Raw pointer to the conversion settings struct. */
-	struct FOpenColorIOColorConversionSettings* ColorSpaceConversion = nullptr;
+	/** Property handle to the conversion settings struct. */
+	TSharedPtr<IPropertyHandle> ColorSpaceSettingsProperty;
 };
diff --git a/Engine/Plugins/Compositing/OpenColorIO/Source/OpenColorIOEditor/Private/OpenColorIOEditorBlueprintLibrary.cpp b/Engine/Plugins/Compositing/OpenColorIO/Source/OpenColorIOEditor/Private/OpenColorIOEditorBlueprintLibrary.cpp
--- a/Engine/Plugins/Compositing/OpenColorIO/Source/OpenColorIOEditor/Private/OpenColorIOEditorBlueprintLibrary.cpp
+++ b/Engine/Plugins/Compositing/OpenColorIO/Source/OpenColorIOEditor/Private/OpenColorIOEditorBlueprintLibrary.cpp
@@ -20,17 +20,17 @@ void UOpenColorIOEditorBlueprintLibrary::SetActiveViewportConfiguration(const FO
 	OpenColorIOEditorModule.SetActiveViewportConfiguration(InConfiguration);
 }
 
-bool UOpenColorIOEditorBlueprintLibrary::ApplyColorSpaceTransformToTexture(const FOpenColorIOColorConversionSettings& ConversionSettings, UTexture* InOutTexture, bool bSynchronous)
+bool UOpenColorIOEditorBlueprintLibrary::ApplyColorSpaceTransformToTexture(const FOpenColorIOColorConversionSettings& ConversionSettings, UTexture* InOutTexture)
 {
 	if (IsValid(InOutTexture))
 	{
-		return ApplyColorSpaceCompressionTransformToTexture(ConversionSettings, InOutTexture->CompressionSettings, InOutTexture, bSynchronous);
+		return ApplyColorSpaceTransformToTextureCompressed(ConversionSettings, InOutTexture->CompressionSettings, InOutTexture);
 	}
 
 	return false;
 }
 
-bool UOpenColorIOEditorBlueprintLibrary::ApplyColorSpaceCompressionTransformToTexture(const FOpenColorIOColorConversionSettings& ConversionSettings, TextureCompressionSettings TargetCompression, UTexture* InOutTexture, bool bSynchronous)
+bool UOpenColorIOEditorBlueprintLibrary::ApplyColorSpaceTransformToTextureCompressed(const FOpenColorIOColorConversionSettings& ConversionSettings, TextureCompressionSettings TargetCompression, UTexture* InOutTexture)
 {
 	if (IsValid(InOutTexture) && IsValid(ConversionSettings.ConfigurationSource))
 	{
@@ -47,10 +47,7 @@ bool UOpenColorIOEditorBlueprintLibrary::ApplyColorSpaceCompressionTransformToTe
 				InOutTexture->CompressionSettings = TargetCompression;
 				InOutTexture->PostEditChange();
 
-				if (bSynchronous)
-				{
-					FTextureCompilingManager::Get().FinishCompilation({ InOutTexture });
-				}
+				FTextureCompilingManager::Get().FinishCompilation({ InOutTexture });
 
 				return true;
 			}
diff --git a/Engine/Plugins/Compositing/OpenColorIO/Source/OpenColorIOEditor/Public/OpenColorIOEditorBlueprintLibrary.h b/Engine/Plugins/Compositing/OpenColorIO/Source/OpenColorIOEditor/Public/OpenColorIOEditorBlueprintLibrary.h
--- a/Engine/Plugins/Compositing/OpenColorIO/Source/OpenColorIOEditor/Public/OpenColorIOEditorBlueprintLibrary.h
+++ b/Engine/Plugins/Compositing/OpenColorIO/Source/OpenColorIOEditor/Public/OpenColorIOEditorBlueprintLibrary.h
@@ -28,11 +28,10 @@ public:
 	 *
 	 * @param ConversionSettings Color transformation settings.
 	 * @param InOutTexture Texture object to transform.
-	 * @param bSynchronous Whether the texture transform is a nlocking operation.
 	 * @return true upon success.
 	 */
 	UFUNCTION(BlueprintCallable, Category = "OpenColorIO")
-	static OPENCOLORIOEDITOR_API bool ApplyColorSpaceTransformToTexture(const FOpenColorIOColorConversionSettings& ConversionSettings, UTexture* InOutTexture, bool bSynchronous = false);
+	static OPENCOLORIOEDITOR_API bool ApplyColorSpaceTransformToTexture(const FOpenColorIOColorConversionSettings& ConversionSettings, UTexture* InOutTexture);
 
 	/**
 	 * Apply a color space transform with a target compression setting to a texture asset.
@@ -40,9 +39,8 @@ public:
 	 * @param ConversionSettings Color transformation settings.
 	 * @param TargetCompression Target texture compression setting.
 	 * @param InOutTexture Texture object to transform.
-	 * @param bSynchronous Whether the texture transform is a nlocking operation.
 	 * @return true upon success.
 	 */
 	UFUNCTION(BlueprintCallable, Category = "OpenColorIO")
-	static OPENCOLORIOEDITOR_API bool ApplyColorSpaceCompressionTransformToTexture(const FOpenColorIOColorConversionSettings& ConversionSettings, TextureCompressionSettings TargetCompression, UTexture* InOutTexture, bool bSynchronous = false);
+	static OPENCOLORIOEDITOR_API bool ApplyColorSpaceTransformToTextureCompressed(const FOpenColorIOColorConversionSettings& ConversionSettings, TextureCompressionSettings TargetCompression, UTexture* InOutTexture);
 };
