OCIO: Fix menu crash on nullptr.

#rnx

diff --git a/Engine/Plugins/Compositing/OpenColorIO/Source/OpenColorIOEditor/Private/OpenColorIOColorSpaceConversionCustomization.cpp b/Engine/Plugins/Compositing/OpenColorIO/Source/OpenColorIOEditor/Private/OpenColorIOColorSpaceConversionCustomization.cpp
--- a/Engine/Plugins/Compositing/OpenColorIO/Source/OpenColorIOEditor/Private/OpenColorIOColorSpaceConversionCustomization.cpp
+++ b/Engine/Plugins/Compositing/OpenColorIO/Source/OpenColorIOEditor/Private/OpenColorIOColorSpaceConversionCustomization.cpp
@@ -80,12 +80,13 @@ void FOpenColorIOColorConversionSettingsCustomization::CustomizeChildren(TShared
 
 				ChildHandle->SetOnPropertyValueChanged(FSimpleDelegate::CreateLambda([this]
 					{
-						FOpenColorIOColorConversionSettings* ColorSpaceConversion = GetConversionSettings();
-
-						TransformPicker[OCIO_Src]->SetConfiguration(ColorSpaceConversion->ConfigurationSource);
-						TransformPicker[OCIO_Dst]->SetConfiguration(ColorSpaceConversion->ConfigurationSource);
+						if (FOpenColorIOColorConversionSettings* ColorSpaceConversion = GetConversionSettings())
+						{
+							TransformPicker[OCIO_Src]->SetConfiguration(ColorSpaceConversion->ConfigurationSource);
+							TransformPicker[OCIO_Dst]->SetConfiguration(ColorSpaceConversion->ConfigurationSource);
 
-						ColorSpaceConversion->OnConversionSettingsChanged().Broadcast();
+							ColorSpaceConversion->OnConversionSettingsChanged().Broadcast();
+						}
 					}));
 			}
 			else if (ChildHandle->GetProperty()->GetFName() == GET_MEMBER_NAME_CHECKED(FOpenColorIOColorConversionSettings, SourceColorSpace))
