OCIO: Ensure that the config hash also takes into account the bSupportInverseViewTransforms setting.

#rnx

diff --git a/Engine/Plugins/Compositing/OpenColorIO/Source/OpenColorIO/Private/OpenColorIOConfiguration.cpp b/Engine/Plugins/Compositing/OpenColorIO/Source/OpenColorIO/Private/OpenColorIOConfiguration.cpp
--- a/Engine/Plugins/Compositing/OpenColorIO/Source/OpenColorIO/Private/OpenColorIOConfiguration.cpp
+++ b/Engine/Plugins/Compositing/OpenColorIO/Source/OpenColorIO/Private/OpenColorIOConfiguration.cpp
@@ -162,7 +162,13 @@ void UOpenColorIOConfiguration::ReloadExistingColorspaces()
 
 	if (Config && Config->IsValid())
 	{
-		const FString LoadedConfigHash = Config->GetCacheID() + FString(OpenColorIOWrapper::GetVersion());
+		FString LoadedConfigHash = Config->GetCacheID() + FString(OpenColorIOWrapper::GetVersion());
+		
+		const UOpenColorIOSettings* Settings = GetDefault<UOpenColorIOSettings>();
+		if (Settings->bSupportInverseViewTransforms)
+		{
+			LoadedConfigHash += FString(TEXT("_Inv"));
+		}
 
 		// Hash is different, proceed with the regeneration...
 		if (ConfigHash != LoadedConfigHash)
@@ -214,8 +220,6 @@ void UOpenColorIOConfiguration::ReloadExistingColorspaces()
 				}
 			}
 
-			const UOpenColorIOSettings* Settings = GetDefault<UOpenColorIOSettings>();
-
 			// Genereate new shaders.
 			for (int32 indexTop = 0; indexTop < DesiredColorSpaces.Num(); ++indexTop)
 			{
