OCIO: Remove cached state in slate preventing real-time updates from scripts.

* Further simplifications to inverted view transforms support (if plugin setting is enabled): simply exposed with a checkbox now.
* Replaced <Invalid> with <None> as requested.
* Enforce forward view transforms if the inverse setting is disabled.

diff --git a/Engine/Plugins/Compositing/OpenColorIO/Source/OpenColorIO/Private/OpenColorIOColorSpace.cpp b/Engine/Plugins/Compositing/OpenColorIO/Source/OpenColorIO/Private/OpenColorIOColorSpace.cpp
--- a/Engine/Plugins/Compositing/OpenColorIO/Source/OpenColorIO/Private/OpenColorIOColorSpace.cpp
+++ b/Engine/Plugins/Compositing/OpenColorIO/Source/OpenColorIO/Private/OpenColorIOColorSpace.cpp
@@ -3,6 +3,7 @@
 #include "OpenColorIOColorSpace.h"
 
 #include "OpenColorIOConfiguration.h"
+#include "OpenColorIOSettings.h"
 
 #include UE_INLINE_GENERATED_CPP_BY_NAME(OpenColorIOColorSpace)
 
@@ -28,7 +29,7 @@ FString FOpenColorIOColorSpace::ToString() const
 	{
 		return ColorSpaceName;
 	}
-	return TEXT("<Invalid>");
+	return {};
 }
 
 bool FOpenColorIOColorSpace::IsValid() const
@@ -88,7 +89,7 @@ FString FOpenColorIODisplayView::ToString() const
 		return Display + TEXT(" - ") + View;
 	}
 
-	return TEXT("<Invalid>");
+	return {};
 }
 
 bool FOpenColorIODisplayView::IsValid() const
@@ -108,10 +109,28 @@ void FOpenColorIODisplayView::Reset()
 
 FOpenColorIOColorConversionSettings::FOpenColorIOColorConversionSettings()
 	: ConfigurationSource(nullptr)
+	, SourceColorSpace()
+	, DestinationColorSpace()
+	, DestinationDisplayView()
+	, DisplayViewDirection(EOpenColorIOViewTransformDirection::Forward)
 {
 
 }
 
+void FOpenColorIOColorConversionSettings::PostSerialize(const FArchive& Ar)
+{
+#if WITH_EDITORONLY_DATA
+	if (Ar.IsLoading())
+	{
+		if (!GetDefault<UOpenColorIOSettings>()->bSupportInverseViewTransforms)
+		{
+			// Enforce forward direction only
+			DisplayViewDirection = EOpenColorIOViewTransformDirection::Forward;
+		}
+	}
+#endif
+}
+
 FString FOpenColorIOColorConversionSettings::ToString() const
 {
 	if (::IsValid(ConfigurationSource))
@@ -145,16 +164,7 @@ bool FOpenColorIOColorConversionSettings::IsValid() const
 		{
 			if (SourceColorSpace.IsValid() && DestinationDisplayView.IsValid())
 			{
-				switch (DisplayViewDirection)
-				{
-				case EOpenColorIOViewTransformDirection::Forward:
-					return ConfigurationSource->HasTransform(SourceColorSpace.ColorSpaceName, DestinationDisplayView.Display, DestinationDisplayView.View, EOpenColorIOViewTransformDirection::Forward);
-				case EOpenColorIOViewTransformDirection::Inverse:
-					return ConfigurationSource->HasTransform(SourceColorSpace.ColorSpaceName, DestinationDisplayView.Display, DestinationDisplayView.View, EOpenColorIOViewTransformDirection::Inverse);
-				default:
-					checkNoEntry();
-					return false;
-				}
+				return ConfigurationSource->HasTransform(SourceColorSpace.ColorSpaceName, DestinationDisplayView.Display, DestinationDisplayView.View, DisplayViewDirection);
 			}
 		}
 		else
@@ -169,6 +179,43 @@ bool FOpenColorIOColorConversionSettings::IsValid() const
 	return false;
 }
 
+FString FOpenColorIOColorConversionSettings::GetSourceString() const
+{
+	if (SourceColorSpace.IsValid())
+	{
+		return SourceColorSpace.ToString();
+	}
+
+	return {};
+}
+
+FString FOpenColorIOColorConversionSettings::GetDestinationString() const
+{
+	if (IsDisplayView() && DestinationDisplayView.IsValid())
+	{
+		return DestinationDisplayView.ToString();
+	}
+	else if (DestinationColorSpace.IsValid())
+	{
+		return DestinationColorSpace.ToString();
+	}
+
+	return {};
+}
+
+void FOpenColorIOColorConversionSettings::Reset(bool bResetConfigurationSource)
+{
+	if (bResetConfigurationSource)
+	{
+		ConfigurationSource = nullptr;
+	}
+	
+	SourceColorSpace.Reset();
+	DestinationColorSpace.Reset();
+	DestinationDisplayView.Reset();
+	DisplayViewDirection = EOpenColorIOViewTransformDirection::Forward;
+}
+
 void FOpenColorIOColorConversionSettings::ValidateColorSpaces()
 {
 	if (::IsValid(ConfigurationSource))
@@ -188,9 +235,7 @@ void FOpenColorIOColorConversionSettings::ValidateColorSpaces()
 	}
 	else
 	{
-		SourceColorSpace.Reset();
-		DestinationColorSpace.Reset();
-		DestinationDisplayView.Reset();
+		Reset();
 	}
 }
 
diff --git a/Engine/Plugins/Compositing/OpenColorIO/Source/OpenColorIO/Public/OpenColorIOColorSpace.h b/Engine/Plugins/Compositing/OpenColorIO/Source/OpenColorIO/Public/OpenColorIOColorSpace.h
--- a/Engine/Plugins/Compositing/OpenColorIO/Source/OpenColorIO/Public/OpenColorIOColorSpace.h
+++ b/Engine/Plugins/Compositing/OpenColorIO/Source/OpenColorIO/Public/OpenColorIOColorSpace.h
@@ -149,9 +149,12 @@ public:
 	EOpenColorIOViewTransformDirection DisplayViewDirection = EOpenColorIOViewTransformDirection::Forward;
 
 	/** Delegate triggered upon changes to the settings. */
+	UE_DEPRECATED(5.3, "This delegate is deprecated.")
 	FOnConversionSettingsChange& OnConversionSettingsChanged() { return ConversionSettingsChanged; }
 public:
 
+	void PostSerialize(const FArchive& Ar);
+
 	/**
 	 * Get a string representation of this conversion.
 	 * @return String representation, i.e. "ConfigurationAssetName - SourceColorSpace to DestinationColorSpace".
@@ -163,6 +166,18 @@ public:
 	 */
 	bool IsValid() const;
 
+	/** Returns a string representing the settings' source */
+	FString GetSourceString() const;
+
+	/** Returns a string representing the settings' destination */
+	FString GetDestinationString() const;
+
+	/**
+	* Reset members to default/empty values.
+	* @param InDepth Desired depth in the family string. 0 == First layer. 
+	*/
+	void Reset(bool bResetConfigurationSource = false);
+
 	/**
 	* Ensure that the selected source and destination color spaces are valid, resets them otherwise.
 	*/
@@ -189,6 +204,14 @@ private:
 	FOnConversionSettingsChange ConversionSettingsChanged;
 };
 
+template<> struct TStructOpsTypeTraits<FOpenColorIOColorConversionSettings> : public TStructOpsTypeTraitsBase2<FOpenColorIOColorConversionSettings>
+{
+	enum
+	{
+		WithPostSerialize = true,
+	};
+};
+
 /**
  * Identifies an OCIO Display look configuration 
  */
diff --git a/Engine/Plugins/Compositing/OpenColorIO/Source/OpenColorIOEditor/Private/OpenColorIOColorSpaceConversionCustomization.cpp b/Engine/Plugins/Compositing/OpenColorIO/Source/OpenColorIOEditor/Private/OpenColorIOColorSpaceConversionCustomization.cpp
--- a/Engine/Plugins/Compositing/OpenColorIO/Source/OpenColorIOEditor/Private/OpenColorIOColorSpaceConversionCustomization.cpp
+++ b/Engine/Plugins/Compositing/OpenColorIO/Source/OpenColorIOEditor/Private/OpenColorIOColorSpaceConversionCustomization.cpp
@@ -82,10 +82,10 @@ void FOpenColorIOColorConversionSettingsCustomization::CustomizeChildren(TShared
 					{
 						if (FOpenColorIOColorConversionSettings* ColorSpaceConversion = GetConversionSettings())
 						{
+							ColorSpaceConversion->Reset();
+
 							TransformPicker[OCIO_Src]->SetConfiguration(ColorSpaceConversion->ConfigurationSource);
 							TransformPicker[OCIO_Dst]->SetConfiguration(ColorSpaceConversion->ConfigurationSource);
-
-							ColorSpaceConversion->OnConversionSettingsChanged().Broadcast();
 						}
 					}));
 			}
@@ -114,29 +114,51 @@ void FOpenColorIOColorConversionSettingsCustomization::CustomizeChildren(TShared
 			ParentHandle->SetOnPropertyResetToDefault(FSimpleDelegate::CreateSP(this, &FOpenColorIOColorConversionSettingsCustomization::OnConfigurationReset));
 		}
 
-		ApplyConfigurationToSelection();
-
 		// Source color space picker widget
 		TransformPicker[OCIO_Src] =
 			SNew(SOpenColorIOColorSpacePicker)
 			.Config(ColorSpaceConversion->ConfigurationSource)
-			.InitialColorSpace(TransformSelection[OCIO_Src].ColorSpace)
-			.RestrictedColor(TransformSelection[OCIO_Dst].ColorSpace)
-			.InitialDisplayView(TransformSelection[OCIO_Src].DisplayView)
-			.RestrictedDisplayView(TransformSelection[OCIO_Dst].DisplayView)
+			.Selection_Lambda([this]() -> FText {
+					if (FOpenColorIOColorConversionSettings* ColorSpaceConversion = GetConversionSettings())
+					{
+						return FText::FromString(ColorSpaceConversion->GetSourceString());
+					}
+					
+					return {};
+				})
+			.SelectionRestriction_Lambda([this]() -> FString {
+					if (FOpenColorIOColorConversionSettings* ColorSpaceConversion = GetConversionSettings())
+					{
+						return ColorSpaceConversion->GetDestinationString();
+					}
+					
+					return {};
+				})
 			.IsDestination(false)
-			.OnColorSpaceChanged(FOnColorSpaceChanged::CreateSP(this, &FOpenColorIOColorConversionSettingsCustomization::OnSourceColorSpaceChanged));
+			.OnColorSpaceChanged(FOnColorSpaceChanged::CreateSP(this, &FOpenColorIOColorConversionSettingsCustomization::OnSelectionChanged));
 
 		// Destination color space picker widget
 		TransformPicker[OCIO_Dst] =
 			SNew(SOpenColorIOColorSpacePicker)
 			.Config(ColorSpaceConversion->ConfigurationSource)
-			.InitialColorSpace(TransformSelection[OCIO_Dst].ColorSpace)
-			.RestrictedColor(TransformSelection[OCIO_Src].ColorSpace)
-			.InitialDisplayView(TransformSelection[OCIO_Dst].DisplayView)
-			.RestrictedDisplayView(TransformSelection[OCIO_Src].DisplayView)
+			.Selection_Lambda([this]() -> FText {
+					if (FOpenColorIOColorConversionSettings* ColorSpaceConversion = GetConversionSettings())
+					{
+						return FText::FromString(ColorSpaceConversion->GetDestinationString());
+					}
+					
+					return {};
+				})
+			.SelectionRestriction_Lambda([this]() -> FString {
+					if (FOpenColorIOColorConversionSettings* ColorSpaceConversion = GetConversionSettings())
+					{
+						return ColorSpaceConversion->GetSourceString();
+					}
+
+					return {};
+				})
 			.IsDestination(true)
-			.OnColorSpaceChanged(FOnColorSpaceChanged::CreateSP(this, &FOpenColorIOColorConversionSettingsCustomization::OnDestinationColorSpaceChanged));
+			.OnColorSpaceChanged(FOnColorSpaceChanged::CreateSP(this, &FOpenColorIOColorConversionSettingsCustomization::OnSelectionChanged));
 
 		// Source color space picker widget
 		StructBuilder.AddCustomRow(LOCTEXT("TransformSource", "Transform Source"))
@@ -166,65 +188,60 @@ void FOpenColorIOColorConversionSettingsCustomization::CustomizeChildren(TShared
 		[
 			TransformPicker[OCIO_Dst].ToSharedRef()
 		];
-	}
-
-	InStructPropertyHandle->MarkHiddenByCustomization();
-}
-
-void FOpenColorIOColorConversionSettingsCustomization::OnSourceColorSpaceChanged(const FOpenColorIOColorSpace& NewColorSpace, const FOpenColorIODisplayView& NewDisplayView)
-{
-	TransformSelection[OCIO_Src].ColorSpace = NewColorSpace;
-	TransformSelection[OCIO_Src].DisplayView = NewDisplayView;
-
-	TransformPicker[OCIO_Dst]->SetRestrictions(NewColorSpace, NewDisplayView);
 
-	ApplySelectionToConfiguration();
-
-	if (FOpenColorIOColorConversionSettings* ColorSpaceConversion = GetConversionSettings())
-	{
-		ColorSpaceConversion->OnConversionSettingsChanged().Broadcast();
-	}
-}
-
-void FOpenColorIOColorConversionSettingsCustomization::OnDestinationColorSpaceChanged(const FOpenColorIOColorSpace& NewColorSpace, const FOpenColorIODisplayView& NewDisplayView)
-{
-	TransformSelection[OCIO_Dst].ColorSpace = NewColorSpace;
-	TransformSelection[OCIO_Dst].DisplayView = NewDisplayView;
-
-	TransformPicker[OCIO_Src]->SetRestrictions(NewColorSpace, NewDisplayView);
-
-	ApplySelectionToConfiguration();
+		const UOpenColorIOSettings* Settings = GetDefault<UOpenColorIOSettings>();
+		if (Settings->bSupportInverseViewTransforms)
+		{
+			StructBuilder.AddCustomRow(LOCTEXT("InvertViewTransform", "Invert View Transform"))
+			.NameContent()
+			[
+				SNew(STextBlock)
+				.Text(LOCTEXT("InvertViewTransform", "Invert View Transform"))
+				.ToolTipText(LOCTEXT("InvertViewTransform_Tooltip", "Option to invert the display-view transform."))
+				.Font(StructCustomizationUtils.GetRegularFont())
+			]
+			.ValueContent()
+			.MaxDesiredWidth(512)
+			[
+				SNew(SCheckBox)
+				.IsChecked_Lambda([this]() -> ECheckBoxState
+				{
+					if (FOpenColorIOColorConversionSettings* ColorSpaceConversion = GetConversionSettings())
+					{
+						if (ColorSpaceConversion->IsDisplayView())
+						{
+							return ColorSpaceConversion->DisplayViewDirection == EOpenColorIOViewTransformDirection::Inverse ? ECheckBoxState::Checked : ECheckBoxState::Unchecked;
+						}
+					}
 
-	if (FOpenColorIOColorConversionSettings* ColorSpaceConversion = GetConversionSettings())
-	{
-		ColorSpaceConversion->OnConversionSettingsChanged().Broadcast();
-	}
-}
+					return ECheckBoxState::Undetermined;
+				})
+				.OnCheckStateChanged_Lambda([this](ECheckBoxState NewState)
+				{
+					if (DisplayViewDirectionProperty.IsValid())
+					{
+						DisplayViewDirectionProperty->SetValue((NewState == ECheckBoxState::Checked) ? (uint8)EOpenColorIOViewTransformDirection::Inverse : (uint8)EOpenColorIOViewTransformDirection::Forward);
+					}
+				})
+				.IsEnabled_Lambda([this]()
+				{
+					if (FOpenColorIOColorConversionSettings* ColorSpaceConversion = GetConversionSettings())
+					{
+						return ColorSpaceConversion->IsDisplayView();
+					}
 
-void FOpenColorIOColorConversionSettingsCustomization::ApplyConfigurationToSelection()
-{
-	if (FOpenColorIOColorConversionSettings* ColorSpaceConversion = GetConversionSettings())
-	{
-		if (ColorSpaceConversion->DisplayViewDirection == EOpenColorIOViewTransformDirection::Forward)
-		{
-			TransformSelection[OCIO_Src].ColorSpace = ColorSpaceConversion->SourceColorSpace;
-			TransformSelection[OCIO_Src].DisplayView.Reset();
-			TransformSelection[OCIO_Dst].ColorSpace = ColorSpaceConversion->DestinationColorSpace;
-			TransformSelection[OCIO_Dst].DisplayView = ColorSpaceConversion->DestinationDisplayView;
-		}
-		else
-		{
-			TransformSelection[OCIO_Src].ColorSpace.Reset();
-			TransformSelection[OCIO_Src].DisplayView = ColorSpaceConversion->DestinationDisplayView;
-			TransformSelection[OCIO_Dst].ColorSpace = ColorSpaceConversion->SourceColorSpace;
-			TransformSelection[OCIO_Dst].DisplayView.Reset();
+					return false;
+				})
+			];
 		}
 	}
+
+	InStructPropertyHandle->MarkHiddenByCustomization();
 }
 
-void FOpenColorIOColorConversionSettingsCustomization::ApplySelectionToConfiguration()
+void FOpenColorIOColorConversionSettingsCustomization::OnSelectionChanged(const FOpenColorIOColorSpace& NewColorSpace, const FOpenColorIODisplayView& NewDisplayView, bool bIsDestination)
 {
-	TArray<TSharedPtr<IPropertyHandle>> Properties = { SourceColorSpaceProperty, DestinationColorSpaceProperty, DestinationDisplayViewProperty, DisplayViewDirectionProperty };
+	TArray<TSharedPtr<IPropertyHandle>> Properties = { SourceColorSpaceProperty, DestinationColorSpaceProperty, DestinationDisplayViewProperty };
 
 	const FScopedTransaction Transaction(LOCTEXT("OCIOConfigurationSelectionUpdate", "OCIO Configuration Selection Update"));
 
@@ -244,19 +261,22 @@ void FOpenColorIOColorConversionSettingsCustomization::ApplySelectionToConfigura
 
 	if (FOpenColorIOColorConversionSettings* ColorSpaceConversion = GetConversionSettings())
 	{
-		if (TransformSelection[OCIO_Src].DisplayView.IsValid())
+		if (bIsDestination)
 		{
-			ColorSpaceConversion->SourceColorSpace = TransformSelection[OCIO_Dst].ColorSpace;
-			ColorSpaceConversion->DestinationColorSpace.Reset();
-			ColorSpaceConversion->DestinationDisplayView = TransformSelection[OCIO_Src].DisplayView;
-			ColorSpaceConversion->DisplayViewDirection = EOpenColorIOViewTransformDirection::Inverse;
+			if (NewDisplayView.IsValid())
+			{
+				ColorSpaceConversion->DestinationColorSpace.Reset();
+				ColorSpaceConversion->DestinationDisplayView = NewDisplayView;
+			}
+			else
+			{
+				ColorSpaceConversion->DestinationDisplayView.Reset();
+				ColorSpaceConversion->DestinationColorSpace = NewColorSpace;
+			}
 		}
 		else
 		{
-			ColorSpaceConversion->SourceColorSpace = TransformSelection[OCIO_Src].ColorSpace;
-			ColorSpaceConversion->DestinationColorSpace = TransformSelection[OCIO_Dst].ColorSpace;
-			ColorSpaceConversion->DestinationDisplayView = TransformSelection[OCIO_Dst].DisplayView;
-			ColorSpaceConversion->DisplayViewDirection = EOpenColorIOViewTransformDirection::Forward;
+			ColorSpaceConversion->SourceColorSpace = NewColorSpace;
 		}
 	}
 
@@ -277,7 +297,7 @@ void FOpenColorIOColorConversionSettingsCustomization::OnConfigurationReset()
 
 	if (FOpenColorIOColorConversionSettings* ColorSpaceConversion = GetConversionSettings())
 	{
-		ColorSpaceConversion->OnConversionSettingsChanged().Broadcast();
+		ColorSpaceConversion->Reset();
 	}
 }
 
diff --git a/Engine/Plugins/Compositing/OpenColorIO/Source/OpenColorIOEditor/Private/OpenColorIOColorSpaceConversionCustomization.h b/Engine/Plugins/Compositing/OpenColorIO/Source/OpenColorIOEditor/Private/OpenColorIOColorSpaceConversionCustomization.h
--- a/Engine/Plugins/Compositing/OpenColorIO/Source/OpenColorIOEditor/Private/OpenColorIOColorSpaceConversionCustomization.h
+++ b/Engine/Plugins/Compositing/OpenColorIO/Source/OpenColorIOEditor/Private/OpenColorIOColorSpaceConversionCustomization.h
@@ -23,14 +23,7 @@ public:
 	virtual void CustomizeHeader(TSharedRef<class IPropertyHandle> InPropertyHandle, class FDetailWidgetRow& HeaderRow, IPropertyTypeCustomizationUtils& PropertyTypeCustomizationUtils) override;
 	virtual void CustomizeChildren(TSharedRef<class IPropertyHandle> InPropertyHandle, class IDetailChildrenBuilder& StructBuilder, IPropertyTypeCustomizationUtils& PropertyTypeCustomizationUtils) override;
 
-	void OnSourceColorSpaceChanged(const FOpenColorIOColorSpace& NewColorSpace, const FOpenColorIODisplayView& NewDisplayView);
-	void OnDestinationColorSpaceChanged(const FOpenColorIOColorSpace& NewColorSpace, const FOpenColorIODisplayView& NewDisplayView);
-
-protected:
-	/** Read configuration settings into local selection objects.*/
-	void ApplyConfigurationToSelection();
-	/** Apply transform selections onto the configuration object.*/
-	void ApplySelectionToConfiguration();
+	void OnSelectionChanged(const FOpenColorIOColorSpace& NewColorSpace, const FOpenColorIODisplayView& NewDisplayView, bool bIsDestination);
 
 private:
 	/** Callback to reset the configuration in the transform source/destination pickers. */
@@ -53,8 +46,6 @@ private:
 
 	/** ColorSpace pickers reference to update them when config asset is changed */
 	TStaticArray<TSharedPtr<SOpenColorIOColorSpacePicker>, 2> TransformPicker;
-	/** Intermediate selection objects. This is done to allow inverted display-view selections (when enabled in settings). */
-	TStaticArray<FOpenColorIOPickerSelection, 2> TransformSelection;
 
 	/** Property handle to the conversion settings struct. */
 	TSharedPtr<IPropertyHandle> ColorSpaceSettingsProperty;
diff --git a/Engine/Plugins/Compositing/OpenColorIO/Source/OpenColorIOEditor/Private/OpenColorIOColorSpaceCustomization.cpp b/Engine/Plugins/Compositing/OpenColorIO/Source/OpenColorIOEditor/Private/OpenColorIOColorSpaceCustomization.cpp
--- a/Engine/Plugins/Compositing/OpenColorIO/Source/OpenColorIOEditor/Private/OpenColorIOColorSpaceCustomization.cpp
+++ b/Engine/Plugins/Compositing/OpenColorIO/Source/OpenColorIOEditor/Private/OpenColorIOColorSpaceCustomization.cpp
@@ -329,7 +329,7 @@ TSharedRef<SWidget> FOpenColorIODisplayViewCustomization::HandleSourceComboButto
 
 			for (int32 i = 0; i < DisplayCount; ++i)
 			{
-				const FOpenColorIODisplayView DisplayViewValue = { ConfigWrapper->GetDisplayName(i), TEXT("<Invalid>")};
+				const FOpenColorIODisplayView DisplayViewValue = { ConfigWrapper->GetDisplayName(i), FStringView{} };
 
 				MenuBuilder.AddSubMenu(
 					FText::FromString(DisplayViewValue.Display),
diff --git a/Engine/Plugins/Compositing/OpenColorIO/Source/OpenColorIOEditor/Private/Widgets/SOpenColorIOColorSpacePicker.cpp b/Engine/Plugins/Compositing/OpenColorIO/Source/OpenColorIOEditor/Private/Widgets/SOpenColorIOColorSpacePicker.cpp
--- a/Engine/Plugins/Compositing/OpenColorIO/Source/OpenColorIOEditor/Private/Widgets/SOpenColorIOColorSpacePicker.cpp
+++ b/Engine/Plugins/Compositing/OpenColorIO/Source/OpenColorIOEditor/Private/Widgets/SOpenColorIOColorSpacePicker.cpp
@@ -19,11 +19,9 @@
 void SOpenColorIOColorSpacePicker::Construct(const FArguments& InArgs)
 {
 	Configuration = InArgs._Config;
-	ColorSpaceSelection = InArgs._InitialColorSpace;
-	RestrictedColorSpace = InArgs._RestrictedColor;
-	DisplayViewSelection = InArgs._InitialDisplayView;
-	RestrictedDisplayView = InArgs._RestrictedDisplayView;
 	bIsDestination = InArgs._IsDestination;
+	Selection = InArgs._Selection;
+	SelectionRestriction = InArgs._SelectionRestriction;
 	OnColorSpaceChanged = InArgs._OnColorSpaceChanged;
 
 	SelectionButton = SNew(SComboButton)
@@ -31,7 +29,7 @@ void SOpenColorIOColorSpacePicker::Construct(const FArguments& InArgs)
 		.ContentPadding(FMargin(4.0, 2.0));
 
 	ChildSlot
-	[
+		[
 		SNew(SHorizontalBox)
 		+ SHorizontalBox::Slot()
 		.FillWidth(1.0f)
@@ -40,18 +38,12 @@ void SOpenColorIOColorSpacePicker::Construct(const FArguments& InArgs)
 			SNew(STextBlock)
 			.Text(MakeAttributeLambda([this]
 				{
-					if (DisplayViewSelection.IsValid())
-					{
-						return FText::FromString(*DisplayViewSelection.ToString());
-					}
-					else if (ColorSpaceSelection.ColorSpaceIndex != INDEX_NONE)
+					if (Selection.IsSet() && !Selection.Get().IsEmpty())
 					{
-						return FText::FromString(*ColorSpaceSelection.ToString());
-					}
-					else
-					{
-						return FText::FromString(TEXT("<Invalid>"));
+						return Selection.Get();
 					}
+					
+					return LOCTEXT("None", "<None>");
 				}))
 		]
 		+ SHorizontalBox::Slot()
@@ -82,25 +74,12 @@ void SOpenColorIOColorSpacePicker::Construct(const FArguments& InArgs)
 void SOpenColorIOColorSpacePicker::SetConfiguration(TWeakObjectPtr<UOpenColorIOConfiguration> NewConfiguration)
 {
 	Configuration = NewConfiguration;
-
-	// Invalidate current color space selection
-	ColorSpaceSelection.Reset();
-	DisplayViewSelection.Reset();
-}
-
-void SOpenColorIOColorSpacePicker::SetRestrictions(const FOpenColorIOColorSpace& InRestrictedColorSpace, const FOpenColorIODisplayView& InRestricedDisplayView)
-{
-	RestrictedColorSpace = InRestrictedColorSpace;
-	RestrictedDisplayView = InRestricedDisplayView;
 }
 
 void SOpenColorIOColorSpacePicker::SetCurrentColorSpace(const FOpenColorIOColorSpace& NewColorSpace)
 {
-	ColorSpaceSelection = NewColorSpace;
-	DisplayViewSelection.Reset();
-	
 	// Let listeners know selection has changed
-	OnColorSpaceChanged.ExecuteIfBound(ColorSpaceSelection, DisplayViewSelection);
+	OnColorSpaceChanged.ExecuteIfBound(NewColorSpace, {}, bIsDestination);
 
 	//Close our menu
 	if (SelectionButton.IsValid())
@@ -111,11 +90,8 @@ void SOpenColorIOColorSpacePicker::SetCurrentColorSpace(const FOpenColorIOColorS
 
 void SOpenColorIOColorSpacePicker::SetCurrentDisplayView(const FOpenColorIODisplayView& NewDisplayView)
 {
-	ColorSpaceSelection.Reset();
-	DisplayViewSelection = NewDisplayView;
-
 	// Let listeners know selection has changed
-	OnColorSpaceChanged.ExecuteIfBound(ColorSpaceSelection, DisplayViewSelection);
+	OnColorSpaceChanged.ExecuteIfBound({}, NewDisplayView, bIsDestination);
 
 	//Close our menu
 	if (SelectionButton.IsValid())
@@ -128,8 +104,6 @@ TSharedRef<SWidget> SOpenColorIOColorSpacePicker::HandleColorSpaceComboButtonMen
 {
 	if (UOpenColorIOConfiguration* ConfigurationObject = Configuration.Get())
 	{
-		const UOpenColorIOSettings* Settings = GetDefault<UOpenColorIOSettings>();
-
 		// generate menu
 		const bool bShouldCloseWindowAfterClosing = false;
 		FMenuBuilder MenuBuilder(bShouldCloseWindowAfterClosing, nullptr);
@@ -141,15 +115,21 @@ TSharedRef<SWidget> SOpenColorIOColorSpacePicker::HandleColorSpaceComboButtonMen
 			for (int32 i = 0; i < ConfigurationObject->DesiredColorSpaces.Num(); ++i)
 			{
 				const FOpenColorIOColorSpace& ColorSpace = ConfigurationObject->DesiredColorSpaces[i];
-				if (ColorSpace == RestrictedColorSpace || !ColorSpace.IsValid())
+				if (!ColorSpace.IsValid())
+				{
+					continue;
+				}
+
+				const FString ColorSpaceStr = ColorSpace.ToString();
+				if (ColorSpaceStr == SelectionRestriction.Get())
 				{
 					continue;
 				}
 
 				MenuBuilder.AddMenuEntry
 				(
-					FText::FromString(ColorSpace.ToString()),
-					FText::FromString(ColorSpace.ToString()),
+					FText::FromString(ColorSpaceStr),
+					FText::FromString(ColorSpaceStr),
 					FSlateIcon(),
 					FUIAction
 					(
@@ -158,9 +138,9 @@ TSharedRef<SWidget> SOpenColorIOColorSpacePicker::HandleColorSpaceComboButtonMen
 							SetCurrentColorSpace(ColorSpace);
 						}),
 						FCanExecuteAction(),
-						FIsActionChecked::CreateLambda([this, ColorSpace]
+						FIsActionChecked::CreateLambda([this, ColorSpaceStr]
 						{
-							return ColorSpaceSelection == ColorSpace;
+							return Selection.Get().ToString() == ColorSpaceStr;
 						})
 						),
 						NAME_None,
@@ -177,7 +157,7 @@ TSharedRef<SWidget> SOpenColorIOColorSpacePicker::HandleColorSpaceComboButtonMen
 		}
 		MenuBuilder.EndSection();
 
-		if (bIsDestination || Settings->bSupportInverseViewTransforms)
+		if (bIsDestination)
 		{
 			MenuBuilder.BeginSection("AvailableDisplayViews", LOCTEXT("AvailableDisplayViews", "Available Display-Views"));
 			{
@@ -186,15 +166,21 @@ TSharedRef<SWidget> SOpenColorIOColorSpacePicker::HandleColorSpaceComboButtonMen
 				for (int32 i = 0; i < ConfigurationObject->DesiredDisplayViews.Num(); ++i)
 				{
 					const FOpenColorIODisplayView& DisplayView = ConfigurationObject->DesiredDisplayViews[i];
-					if (DisplayView == RestrictedDisplayView || !DisplayView.IsValid())
+					if (!DisplayView.IsValid())
+					{
+						continue;
+					}
+					
+					const FString DisplayViewStr = DisplayView.ToString();
+					if (DisplayViewStr == SelectionRestriction.Get())
 					{
 						continue;
 					}
 
 					MenuBuilder.AddMenuEntry
 					(
-						FText::FromString(DisplayView.ToString()),
-						FText::FromString(DisplayView.ToString()),
+						FText::FromString(DisplayViewStr),
+						FText::FromString(DisplayViewStr),
 						FSlateIcon(),
 						FUIAction
 						(
@@ -203,9 +189,9 @@ TSharedRef<SWidget> SOpenColorIOColorSpacePicker::HandleColorSpaceComboButtonMen
 								SetCurrentDisplayView(DisplayView);
 							}),
 							FCanExecuteAction(),
-							FIsActionChecked::CreateLambda([this, DisplayView]
+							FIsActionChecked::CreateLambda([this, DisplayViewStr]
 							{
-								return DisplayViewSelection == DisplayView;
+								return Selection.Get().ToString() == DisplayViewStr;
 							})
 						),
 						NAME_None,
@@ -232,14 +218,21 @@ TSharedRef<SWidget> SOpenColorIOColorSpacePicker::HandleColorSpaceComboButtonMen
 
 FReply SOpenColorIOColorSpacePicker::OnResetToDefault()
 {
-	SetCurrentColorSpace(FOpenColorIOColorSpace());
-	SetCurrentDisplayView(FOpenColorIODisplayView());
+	// Let listeners know selection has changed
+	OnColorSpaceChanged.ExecuteIfBound({}, {}, bIsDestination);
+
+	//Close our menu
+	if (SelectionButton.IsValid())
+	{
+		SelectionButton->SetIsOpen(false);
+	}
+
 	return FReply::Handled();
 }
 
 EVisibility SOpenColorIOColorSpacePicker::ShouldShowResetToDefaultButton() const
 {
-	return (ColorSpaceSelection.IsValid() || DisplayViewSelection.IsValid()) ? EVisibility::Visible : EVisibility::Hidden;
+	return (Selection.IsSet() && !Selection.Get().IsEmpty()) ? EVisibility::Visible : EVisibility::Hidden;
 }
 
 #undef LOCTEXT_NAMESPACE
diff --git a/Engine/Plugins/Compositing/OpenColorIO/Source/OpenColorIOEditor/Private/Widgets/SOpenColorIOColorSpacePicker.h b/Engine/Plugins/Compositing/OpenColorIO/Source/OpenColorIOEditor/Private/Widgets/SOpenColorIOColorSpacePicker.h
--- a/Engine/Plugins/Compositing/OpenColorIO/Source/OpenColorIOEditor/Private/Widgets/SOpenColorIOColorSpacePicker.h
+++ b/Engine/Plugins/Compositing/OpenColorIO/Source/OpenColorIOEditor/Private/Widgets/SOpenColorIOColorSpacePicker.h
@@ -16,7 +16,7 @@ class FViewportClient;
 class SComboButton;
 class UOpenColorIOConfiguration;
 
-DECLARE_DELEGATE_TwoParams(FOnColorSpaceChanged, const FOpenColorIOColorSpace& /*ColorSpace*/, const FOpenColorIODisplayView& /*DisplayView*/);
+DECLARE_DELEGATE_ThreeParams(FOnColorSpaceChanged, const FOpenColorIOColorSpace& /*ColorSpace*/, const FOpenColorIODisplayView& /*DisplayView*/, bool /*bIsDestination*/);
 
 /** Pair of OpenColorIO transform selection objects. */
 struct FOpenColorIOPickerSelection
@@ -37,11 +37,9 @@ class SOpenColorIOColorSpacePicker : public SCompoundWidget
 public:
 	SLATE_BEGIN_ARGS(SOpenColorIOColorSpacePicker) {}
 		SLATE_ARGUMENT(TWeakObjectPtr<UOpenColorIOConfiguration>, Config)
-		SLATE_ARGUMENT(FOpenColorIOColorSpace, InitialColorSpace)
-		SLATE_ARGUMENT(FOpenColorIOColorSpace, RestrictedColor)
-		SLATE_ARGUMENT(FOpenColorIODisplayView, InitialDisplayView)
-		SLATE_ARGUMENT(FOpenColorIODisplayView, RestrictedDisplayView)
 		SLATE_ARGUMENT(bool, IsDestination)
+		SLATE_ATTRIBUTE(FText, Selection)
+		SLATE_ATTRIBUTE(FString, SelectionRestriction)
 		SLATE_EVENT(FOnColorSpaceChanged, OnColorSpaceChanged)
 	SLATE_END_ARGS()
 
@@ -51,9 +49,6 @@ public:
 	/** Update current configuration asset for this picker */
 	void SetConfiguration(TWeakObjectPtr<UOpenColorIOConfiguration> NewConfiguration);
 
-	/** Update restricted color space for this picker */
-	void SetRestrictions(const FOpenColorIOColorSpace& RestrictedColorSpace, const FOpenColorIODisplayView& RestricedDisplayView);
-
 protected:
 	
 	/** Called when a selection has been made */
@@ -72,10 +67,8 @@ protected:
 protected:
 	TSharedPtr<SComboButton> SelectionButton;
 	TWeakObjectPtr<UOpenColorIOConfiguration> Configuration;
-	FOpenColorIOColorSpace ColorSpaceSelection;
-	FOpenColorIOColorSpace RestrictedColorSpace;
-	FOpenColorIODisplayView DisplayViewSelection;
-	FOpenColorIODisplayView RestrictedDisplayView;
+	TAttribute<FText> Selection;
+	TAttribute<FString> SelectionRestriction;
 	FOnColorSpaceChanged OnColorSpaceChanged;
 	bool bIsDestination = false;
 };
diff --git a/Engine/Plugins/Compositing/OpenColorIO/Source/OpenColorIOEditor/Private/Widgets/SOpenColorIODisplay.cpp b/Engine/Plugins/Compositing/OpenColorIO/Source/OpenColorIOEditor/Private/Widgets/SOpenColorIODisplay.cpp
--- a/Engine/Plugins/Compositing/OpenColorIO/Source/OpenColorIOEditor/Private/Widgets/SOpenColorIODisplay.cpp
+++ b/Engine/Plugins/Compositing/OpenColorIO/Source/OpenColorIOEditor/Private/Widgets/SOpenColorIODisplay.cpp
@@ -5,6 +5,7 @@
 #include "AssetRegistry/AssetData.h"
 #include "Framework/Application/SlateApplication.h"
 #include "OpenColorIOConfiguration.h"
+#include "OpenColorIOSettings.h"
 #include "PropertyCustomizationHelpers.h"
 #include "SOpenColorIOColorSpacePicker.h"
 #include "ToolMenu.h"
@@ -61,33 +62,60 @@ void SOpenColorIODisplay::Construct(const FArguments& InArgs)
 		EUserInterfaceActionType::None
 	));
 
-	ApplyConfigurationToSelection();
-
 	// Source color space picker widget
 	TransformPicker[OCIO_Src] =
 		SNew(SOpenColorIOColorSpacePicker)
 		.Config(Configuration.ColorConfiguration.ConfigurationSource)
-		.InitialColorSpace(TransformSelection[OCIO_Src].ColorSpace)
-		.RestrictedColor(TransformSelection[OCIO_Dst].ColorSpace)
-		.InitialDisplayView(TransformSelection[OCIO_Src].DisplayView)
-		.RestrictedDisplayView(TransformSelection[OCIO_Dst].DisplayView)
+		.Selection_Lambda([this]() { return FText::FromString(Configuration.ColorConfiguration.GetSourceString()); })
+		.SelectionRestriction_Lambda([this]() { return Configuration.ColorConfiguration.GetDestinationString(); })
 		.IsDestination(false)
-		.OnColorSpaceChanged(FOnColorSpaceChanged::CreateSP(this, &SOpenColorIODisplay::OnSourceColorSpaceChanged));
+		.OnColorSpaceChanged(FOnColorSpaceChanged::CreateSP(this, &SOpenColorIODisplay::OnSelectionChanged));
 
 	// Destination color space picker widget
 	TransformPicker[OCIO_Dst] =
 		SNew(SOpenColorIOColorSpacePicker)
 		.Config(Configuration.ColorConfiguration.ConfigurationSource)
-		.InitialColorSpace(TransformSelection[OCIO_Dst].ColorSpace)
-		.RestrictedColor(TransformSelection[OCIO_Src].ColorSpace)
-		.InitialDisplayView(TransformSelection[OCIO_Dst].DisplayView)
-		.RestrictedDisplayView(TransformSelection[OCIO_Src].DisplayView)
+		.Selection_Lambda([this]() { return FText::FromString(Configuration.ColorConfiguration.GetDestinationString()); })
+		.SelectionRestriction_Lambda([this]() { return Configuration.ColorConfiguration.GetSourceString(); })
 		.IsDestination(true)
-		.OnColorSpaceChanged(FOnColorSpaceChanged::CreateSP(this, &SOpenColorIODisplay::OnDestinationColorSpaceChanged));
+		.OnColorSpaceChanged(FOnColorSpaceChanged::CreateSP(this, &SOpenColorIODisplay::OnSelectionChanged));
 	
 	// Note: FText::GetEmpty() would previously result in incorrect indentation.
 	Section.AddEntry(FToolMenuEntry::InitWidget("SourceColor", TransformPicker[OCIO_Src].ToSharedRef(), FText::FromString(TEXT(" ")), false, false));
 	Section.AddEntry(FToolMenuEntry::InitWidget("DestinationColor", TransformPicker[OCIO_Dst].ToSharedRef(), FText::FromString(TEXT(" ")), false, false));
+	
+	const UOpenColorIOSettings* Settings = GetDefault<UOpenColorIOSettings>();
+	if (Settings->bSupportInverseViewTransforms)
+	{
+		TSharedRef<SWidget> InvertCheckbox = SNew(SCheckBox)
+			.IsChecked_Lambda([this]() -> ECheckBoxState
+			{
+				if (Configuration.ColorConfiguration.IsDisplayView())
+				{
+					return Configuration.ColorConfiguration.DisplayViewDirection == EOpenColorIOViewTransformDirection::Inverse ? ECheckBoxState::Checked : ECheckBoxState::Unchecked;
+				}
+				
+				return ECheckBoxState::Undetermined;
+			})
+			.OnCheckStateChanged_Lambda([this](ECheckBoxState NewState)
+			{
+				Configuration.ColorConfiguration.DisplayViewDirection = (NewState == ECheckBoxState::Checked) ? EOpenColorIOViewTransformDirection::Inverse : EOpenColorIOViewTransformDirection::Forward;
+
+				OnDisplayConfigurationChanged.ExecuteIfBound(Configuration);
+			})
+			.IsEnabled_Lambda([this]()
+			{
+				return Configuration.ColorConfiguration.IsDisplayView();
+			});
+
+		Section.AddEntry(
+			FToolMenuEntry::InitWidget(
+				"InvertViewTransform",
+				InvertCheckbox,
+				LOCTEXT("InvertViewTransform", "Invert View Transform")
+			));
+	}
+
 
 	TSharedRef<SWidget> Widget = UToolMenus::Get()->GenerateWidget(Menu);
 	ChildSlot
@@ -153,6 +181,8 @@ void SOpenColorIODisplay::OnConfigSelected(const FAssetData& AssetData)
 	FSlateApplication::Get().DismissMenuByWidget(ClassPicker.ToSharedRef());
 
 	UOpenColorIOConfiguration* NewConfig = Cast<UOpenColorIOConfiguration>(AssetData.GetAsset());
+	
+	Configuration.ColorConfiguration.Reset();
 	Configuration.ColorConfiguration.ConfigurationSource = NewConfig;
 	TransformPicker[OCIO_Src]->SetConfiguration(NewConfig);
 	TransformPicker[OCIO_Dst]->SetConfiguration(NewConfig);
@@ -178,63 +208,26 @@ bool SOpenColorIODisplay::GetDisplayConfigurationState()
 	return Configuration.bIsEnabled;
 }
 
-void SOpenColorIODisplay::OnSourceColorSpaceChanged(const FOpenColorIOColorSpace& NewColorSpace, const FOpenColorIODisplayView& NewDisplayView)
-{
-	TransformSelection[OCIO_Src].ColorSpace = NewColorSpace;
-	TransformSelection[OCIO_Src].DisplayView = NewDisplayView;
-
-	TransformPicker[OCIO_Dst]->SetRestrictions(NewColorSpace, NewDisplayView);
-
-	ApplySelectionToConfiguration();
-}
-
-void SOpenColorIODisplay::OnDestinationColorSpaceChanged(const FOpenColorIOColorSpace& NewColorSpace, const FOpenColorIODisplayView& NewDisplayView)
-{
-	TransformSelection[OCIO_Dst].ColorSpace = NewColorSpace;
-	TransformSelection[OCIO_Dst].DisplayView = NewDisplayView;
-
-	TransformPicker[OCIO_Src]->SetRestrictions(NewColorSpace, NewDisplayView);
-
-	ApplySelectionToConfiguration();
-}
-
-void SOpenColorIODisplay::ApplyConfigurationToSelection()
-{
-	const FOpenColorIOColorConversionSettings& ConversionSettings = Configuration.ColorConfiguration;
-
-	if (ConversionSettings.DisplayViewDirection == EOpenColorIOViewTransformDirection::Forward)
-	{
-		TransformSelection[OCIO_Src].ColorSpace = ConversionSettings.SourceColorSpace;
-		TransformSelection[OCIO_Src].DisplayView.Reset();
-		TransformSelection[OCIO_Dst].ColorSpace = ConversionSettings.DestinationColorSpace;
-		TransformSelection[OCIO_Dst].DisplayView = ConversionSettings.DestinationDisplayView;
-	}
-	else
-	{
-		TransformSelection[OCIO_Src].ColorSpace.Reset();
-		TransformSelection[OCIO_Src].DisplayView = ConversionSettings.DestinationDisplayView;
-		TransformSelection[OCIO_Dst].ColorSpace = ConversionSettings.SourceColorSpace;
-		TransformSelection[OCIO_Dst].DisplayView.Reset();
-	}
-}
-
-void SOpenColorIODisplay::ApplySelectionToConfiguration()
+void SOpenColorIODisplay::OnSelectionChanged(const FOpenColorIOColorSpace& NewColorSpace, const FOpenColorIODisplayView& NewDisplayView, bool bIsDestination)
 {
 	FOpenColorIOColorConversionSettings& ConversionSettings = Configuration.ColorConfiguration;
 
-	if (TransformSelection[OCIO_Src].DisplayView.IsValid())
+	if (bIsDestination)
 	{
-		ConversionSettings.SourceColorSpace = TransformSelection[OCIO_Dst].ColorSpace;
-		ConversionSettings.DestinationColorSpace.Reset();
-		ConversionSettings.DestinationDisplayView = TransformSelection[OCIO_Src].DisplayView;
-		ConversionSettings.DisplayViewDirection = EOpenColorIOViewTransformDirection::Inverse;
+		if (NewDisplayView.IsValid())
+		{
+			ConversionSettings.DestinationColorSpace.Reset();
+			ConversionSettings.DestinationDisplayView = NewDisplayView;
+		}
+		else
+		{
+			ConversionSettings.DestinationDisplayView.Reset();
+			ConversionSettings.DestinationColorSpace = NewColorSpace;
+		}
 	}
 	else
 	{
-		ConversionSettings.SourceColorSpace = TransformSelection[OCIO_Src].ColorSpace;
-		ConversionSettings.DestinationColorSpace = TransformSelection[OCIO_Dst].ColorSpace;
-		ConversionSettings.DestinationDisplayView = TransformSelection[OCIO_Dst].DisplayView;
-		ConversionSettings.DisplayViewDirection = EOpenColorIOViewTransformDirection::Forward;
+		ConversionSettings.SourceColorSpace = NewColorSpace;
 	}
 
 	OnDisplayConfigurationChanged.ExecuteIfBound(Configuration);
diff --git a/Engine/Plugins/Compositing/OpenColorIO/Source/OpenColorIOEditor/Private/Widgets/SOpenColorIODisplay.h b/Engine/Plugins/Compositing/OpenColorIO/Source/OpenColorIOEditor/Private/Widgets/SOpenColorIODisplay.h
--- a/Engine/Plugins/Compositing/OpenColorIO/Source/OpenColorIOEditor/Private/Widgets/SOpenColorIODisplay.h
+++ b/Engine/Plugins/Compositing/OpenColorIO/Source/OpenColorIOEditor/Private/Widgets/SOpenColorIODisplay.h
@@ -40,14 +40,9 @@ protected:
 	void ToggleEnableDisplay();
 	bool CanEnableDisplayConfiguration();
 	bool GetDisplayConfigurationState();
-	void OnSourceColorSpaceChanged(const FOpenColorIOColorSpace& NewColorSpace, const FOpenColorIODisplayView& NewDisplayView);
-	void OnDestinationColorSpaceChanged(const FOpenColorIOColorSpace& NewColorSpace, const FOpenColorIODisplayView& NewDisplayView);
-
+	void OnSelectionChanged(const FOpenColorIOColorSpace& NewColorSpace, const FOpenColorIODisplayView& NewDisplayView, bool bIsDestination);
+	
 protected:
-	/** Read configuration settings into local selection objects.*/
-	void ApplyConfigurationToSelection();
-	/** Apply transform selections onto the configuration object.*/
-	void ApplySelectionToConfiguration();
 
 	FViewport* Viewport;
 
@@ -56,8 +51,6 @@ protected:
 
 	/** ColorSpace pickers reference to update them when config asset is changed */
 	TStaticArray<TSharedPtr<SOpenColorIOColorSpacePicker>, 2> TransformPicker;
-	/** Intermediate selection objects. This is done to allow inverted display-view selections (when enabled in settings). */
-	TStaticArray<FOpenColorIOPickerSelection, 2> TransformSelection;
 
 	/** Current configuration */
 	FOpenColorIODisplayConfiguration Configuration;
