OCIO: Fix non-unity build, and added missing editor settings processing.

#rnx

diff --git a/Engine/Plugins/Compositing/OpenColorIO/Source/OpenColorIOEditor/Private/OpenColorIOColorSpaceConversionCustomization.cpp b/Engine/Plugins/Compositing/OpenColorIO/Source/OpenColorIOEditor/Private/OpenColorIOColorSpaceConversionCustomization.cpp
--- a/Engine/Plugins/Compositing/OpenColorIO/Source/OpenColorIOEditor/Private/OpenColorIOColorSpaceConversionCustomization.cpp
+++ b/Engine/Plugins/Compositing/OpenColorIO/Source/OpenColorIOEditor/Private/OpenColorIOColorSpaceConversionCustomization.cpp
@@ -15,6 +15,7 @@
 #include "ScopedTransaction.h"
 #include "Widgets/SOpenColorIOColorSpacePicker.h"
 #include "Widgets/Text/STextBlock.h"
+#include "Widgets/Input/SCheckBox.h"
 
 #define LOCTEXT_NAMESPACE "OpenColorIOColorSpaceConversionCustomization"
 
diff --git a/Engine/Plugins/Compositing/OpenColorIO/Source/OpenColorIOEditor/Private/OpenColorIOEditorSettings.cpp b/Engine/Plugins/Compositing/OpenColorIO/Source/OpenColorIOEditor/Private/OpenColorIOEditorSettings.cpp
--- a/Engine/Plugins/Compositing/OpenColorIO/Source/OpenColorIOEditor/Private/OpenColorIOEditorSettings.cpp
+++ b/Engine/Plugins/Compositing/OpenColorIO/Source/OpenColorIOEditor/Private/OpenColorIOEditorSettings.cpp
@@ -5,6 +5,7 @@
 #include "IOpenColorIOEditorModule.h"
 #include "Misc/ConfigCacheIni.h"
 #include "Misc/Paths.h"
+#include "OpenColorIOSettings.h"
 
 #include UE_INLINE_GENERATED_CPP_BY_NAME(OpenColorIOEditorSettings)
 
@@ -42,6 +43,15 @@ void UOpenColorIOLevelViewportSettings::PostInitProperties()
 			UE_LOG(LogOpenColorIOEditor, Warning, TEXT("Migrated EditorPerProjectUserSettings OpenColorIO settings to plugin-specific config file: %s."), *GetClass()->GetConfigName());
 		}
 	}
+
+	if (!GetDefault<UOpenColorIOSettings>()->bSupportInverseViewTransforms)
+	{
+		for (FPerViewportDisplaySettingPair& ViewportSetting : ViewportsSettings)
+		{
+			// Enforce forward direction only
+			ViewportSetting.DisplayConfiguration.ColorConfiguration.DisplayViewDirection = EOpenColorIOViewTransformDirection::Forward;
+		}
+	}
 }
 
 void UOpenColorIOLevelViewportSettings::SetViewportSettings(FName ViewportIdentifier, const FOpenColorIODisplayConfiguration& Configuration)
diff --git a/Engine/Plugins/Compositing/OpenColorIO/Source/OpenColorIOEditor/Private/Widgets/SOpenColorIODisplay.cpp b/Engine/Plugins/Compositing/OpenColorIO/Source/OpenColorIOEditor/Private/Widgets/SOpenColorIODisplay.cpp
--- a/Engine/Plugins/Compositing/OpenColorIO/Source/OpenColorIOEditor/Private/Widgets/SOpenColorIODisplay.cpp
+++ b/Engine/Plugins/Compositing/OpenColorIO/Source/OpenColorIOEditor/Private/Widgets/SOpenColorIODisplay.cpp
@@ -13,6 +13,7 @@
 #include "ToolMenuEntry.h"
 #include "ToolMenus.h"
 #include "ToolMenuSection.h"
+#include "Widgets/Input/SCheckBox.h"
 
 
 #define LOCTEXT_NAMESPACE "SOpenColorIODisplay"
