OCIO: Force disable on invalid editor viewport settings.

diff --git a/Engine/Plugins/Compositing/OpenColorIO/Source/OpenColorIOEditor/Private/OpenColorIOEditorSettings.cpp b/Engine/Plugins/Compositing/OpenColorIO/Source/OpenColorIOEditor/Private/OpenColorIOEditorSettings.cpp
--- a/Engine/Plugins/Compositing/OpenColorIO/Source/OpenColorIOEditor/Private/OpenColorIOEditorSettings.cpp
+++ b/Engine/Plugins/Compositing/OpenColorIO/Source/OpenColorIOEditor/Private/OpenColorIOEditorSettings.cpp
@@ -44,13 +44,20 @@ void UOpenColorIOLevelViewportSettings::PostInitProperties()
 		}
 	}
 
-	if (!GetDefault<UOpenColorIOSettings>()->bSupportInverseViewTransforms)
+	const bool bEnforceForwardViewDirectionOnly = !GetDefault<UOpenColorIOSettings>()->bSupportInverseViewTransforms;
+
+	for (FPerViewportDisplaySettingPair& ViewportSetting : ViewportsSettings)
 	{
-		for (FPerViewportDisplaySettingPair& ViewportSetting : ViewportsSettings)
+		if (bEnforceForwardViewDirectionOnly)
 		{
-			// Enforce forward direction only
 			ViewportSetting.DisplayConfiguration.ColorConfiguration.DisplayViewDirection = EOpenColorIOViewTransformDirection::Forward;
 		}
+
+		// Force disable on invalid viewport settings.
+		if (!ViewportSetting.DisplayConfiguration.ColorConfiguration.IsValid())
+		{
+			ViewportSetting.DisplayConfiguration.bIsEnabled = false;
+		}
 	}
 }
 
