OCIO: Minor fixes.

* Make sure we properly display <None> in the config menu when nothing is selected.
* Added log message to force-disable on invalid viewport settings.

#rnx

diff --git a/Engine/Plugins/Compositing/OpenColorIO/Source/OpenColorIOEditor/Private/OpenColorIOColorSpaceCustomization.cpp b/Engine/Plugins/Compositing/OpenColorIO/Source/OpenColorIOEditor/Private/OpenColorIOColorSpaceCustomization.cpp
--- a/Engine/Plugins/Compositing/OpenColorIO/Source/OpenColorIOEditor/Private/OpenColorIOColorSpaceCustomization.cpp
+++ b/Engine/Plugins/Compositing/OpenColorIO/Source/OpenColorIOEditor/Private/OpenColorIOColorSpaceCustomization.cpp
@@ -45,7 +45,16 @@ void FOpenColorIOColorSpaceCustomization::CustomizeHeader(TSharedRef<IPropertyHa
 				.VAlign(VAlign_Center)
 				[
 					SNew(STextBlock)
-					.Text(MakeAttributeLambda([=] { return FText::FromString(ColorSpaceValue->ToString()); }))
+					.Text(MakeAttributeLambda([ColorSpaceValue] {
+							const FString ColorSpaceName = ColorSpaceValue->ToString();
+							
+							if(!ColorSpaceName.IsEmpty())
+							{
+								return FText::FromString(ColorSpaceName);
+							}
+
+							return LOCTEXT("None", "<None>");
+						}))
 				]
 				+ SHorizontalBox::Slot()
 				.AutoWidth()
@@ -93,7 +102,16 @@ void FOpenColorIODisplayViewCustomization::CustomizeHeader(TSharedRef<IPropertyH
 				.VAlign(VAlign_Center)
 				[
 					SNew(STextBlock)
-					.Text(MakeAttributeLambda([=] { return FText::FromString(DisplayViewValue->ToString()); }))
+					.Text(MakeAttributeLambda([DisplayViewValue] {
+							const FString DisplayViewName = DisplayViewValue->ToString();
+							
+							if(!DisplayViewName.IsEmpty())
+							{
+								return FText::FromString(DisplayViewName);
+							}
+
+							return LOCTEXT("None", "<None>");
+						}))
 				]
 				+ SHorizontalBox::Slot()
 				.AutoWidth()
diff --git a/Engine/Plugins/Compositing/OpenColorIO/Source/OpenColorIOEditor/Private/OpenColorIOEditorSettings.cpp b/Engine/Plugins/Compositing/OpenColorIO/Source/OpenColorIOEditor/Private/OpenColorIOEditorSettings.cpp
--- a/Engine/Plugins/Compositing/OpenColorIO/Source/OpenColorIOEditor/Private/OpenColorIOEditorSettings.cpp
+++ b/Engine/Plugins/Compositing/OpenColorIO/Source/OpenColorIOEditor/Private/OpenColorIOEditorSettings.cpp
@@ -53,9 +53,11 @@ void UOpenColorIOLevelViewportSettings::PostInitProperties()
 			ViewportSetting.DisplayConfiguration.ColorConfiguration.DisplayViewDirection = EOpenColorIOViewTransformDirection::Forward;
 		}
 
-		// Force disable on invalid viewport settings.
+		// Note: Ideally the following wouldn't be necessary but it is required since the previous default behavior was enabled with invalid settings.
 		if (!ViewportSetting.DisplayConfiguration.ColorConfiguration.IsValid())
 		{
+			UE_LOG(LogOpenColorIOEditor, Display, TEXT("Force-disable invalid viewport transform settings."));
+
 			ViewportSetting.DisplayConfiguration.bIsEnabled = false;
 		}
 	}
