OCIO: Exposing CPU transformation to colors (and not only images).

diff --git a/Engine/Plugins/Compositing/OpenColorIO/Source/OpenColorIO/Private/OpenColorIOColorTransform.cpp b/Engine/Plugins/Compositing/OpenColorIO/Source/OpenColorIO/Private/OpenColorIOColorTransform.cpp
--- a/Engine/Plugins/Compositing/OpenColorIO/Source/OpenColorIO/Private/OpenColorIOColorTransform.cpp
+++ b/Engine/Plugins/Compositing/OpenColorIO/Source/OpenColorIO/Private/OpenColorIOColorTransform.cpp
@@ -628,6 +628,27 @@ bool UOpenColorIOColorTransform::IsTransform(const FString& InSourceColorSpace,
 }
 
 #if WITH_EDITOR
+bool UOpenColorIOColorTransform::TransformColor(FLinearColor& InOutColor)
+{
+	const UOpenColorIOConfiguration* ConfigurationOwner = Cast<UOpenColorIOConfiguration>(GetOuter());
+	const FOpenColorIOWrapperConfig* ConfigWrapper = GetTransformConfigWrapper(ConfigurationOwner);
+
+	if (ConfigWrapper != nullptr)
+	{
+		const FOpenColorIOWrapperProcessor TransformProcessor = GetTransformProcessor(this, ConfigWrapper);
+		if (TransformProcessor.IsValid())
+		{
+			return TransformProcessor.TransformColor(InOutColor);
+		}
+	}
+	else
+	{
+		UE_LOG(LogOpenColorIO, Error, TEXT("Failed to transform image for color transform %s. Configuration file was invalid."), *GetTransformFriendlyName());
+	}
+
+	return false;
+}
+
 bool UOpenColorIOColorTransform::TransformImage(const FImageView& InOutImage) const
 {
 	const UOpenColorIOConfiguration* ConfigurationOwner = Cast<UOpenColorIOConfiguration>(GetOuter());
diff --git a/Engine/Plugins/Compositing/OpenColorIO/Source/OpenColorIO/Private/OpenColorIOConfiguration.cpp b/Engine/Plugins/Compositing/OpenColorIO/Source/OpenColorIO/Private/OpenColorIOConfiguration.cpp
--- a/Engine/Plugins/Compositing/OpenColorIO/Source/OpenColorIO/Private/OpenColorIOConfiguration.cpp
+++ b/Engine/Plugins/Compositing/OpenColorIO/Source/OpenColorIO/Private/OpenColorIOConfiguration.cpp
@@ -134,6 +134,16 @@ bool UOpenColorIOConfiguration::Validate() const
 }
 
 #if WITH_EDITOR
+bool UOpenColorIOConfiguration::TransformColor(const FOpenColorIOColorConversionSettings& InSettings, FLinearColor& InOutColor) const
+{
+	if (const TObjectPtr<UOpenColorIOColorTransform>* TransformPtr = FindTransform(InSettings))
+	{
+		return (*TransformPtr)->TransformColor(InOutColor);
+	}
+
+	return false;
+}
+
 bool UOpenColorIOConfiguration::TransformImage(const FOpenColorIOColorConversionSettings& InSettings, const FImageView& InOutImage) const
 {
 	if (const TObjectPtr<UOpenColorIOColorTransform>* TransformPtr = FindTransform(InSettings))
diff --git a/Engine/Plugins/Compositing/OpenColorIO/Source/OpenColorIO/Public/OpenColorIOColorTransform.h b/Engine/Plugins/Compositing/OpenColorIO/Source/OpenColorIO/Public/OpenColorIOColorTransform.h
--- a/Engine/Plugins/Compositing/OpenColorIO/Source/OpenColorIO/Public/OpenColorIOColorTransform.h
+++ b/Engine/Plugins/Compositing/OpenColorIO/Source/OpenColorIO/Public/OpenColorIOColorTransform.h
@@ -80,6 +80,9 @@ public:
 	bool IsTransform(const FString& InSourceColorSpace, const FString& InDisplay, const FString& InView, EOpenColorIOViewTransformDirection InDirection) const;
 
 #if WITH_EDITOR
+	/** Apply the color transform in-place to the specified color. */
+	bool TransformColor(FLinearColor& InOutColor);
+		
 	/** Apply the color transform in-place to the specified image. */
 	bool TransformImage(const FImageView& InOutImage) const;
 	
diff --git a/Engine/Plugins/Compositing/OpenColorIO/Source/OpenColorIO/Public/OpenColorIOConfiguration.h b/Engine/Plugins/Compositing/OpenColorIO/Source/OpenColorIO/Public/OpenColorIOConfiguration.h
--- a/Engine/Plugins/Compositing/OpenColorIO/Source/OpenColorIO/Public/OpenColorIOConfiguration.h
+++ b/Engine/Plugins/Compositing/OpenColorIO/Source/OpenColorIO/Public/OpenColorIOConfiguration.h
@@ -46,6 +46,9 @@ public:
 	bool Validate() const;
 
 #if WITH_EDITOR
+	/** Apply the color transform in-place to the specified image. */
+	bool TransformColor(const FOpenColorIOColorConversionSettings& InSettings, FLinearColor& InOutColor) const;
+
 	/** Apply the color transform in-place to the specified image. */
 	bool TransformImage(const FOpenColorIOColorConversionSettings& InSettings, const FImageView& InOutImage) const;
 
diff --git a/Engine/Plugins/Compositing/OpenColorIO/Source/OpenColorIOEditor/Private/OpenColorIOEditorBlueprintLibrary.cpp b/Engine/Plugins/Compositing/OpenColorIO/Source/OpenColorIOEditor/Private/OpenColorIOEditorBlueprintLibrary.cpp
--- a/Engine/Plugins/Compositing/OpenColorIO/Source/OpenColorIOEditor/Private/OpenColorIOEditorBlueprintLibrary.cpp
+++ b/Engine/Plugins/Compositing/OpenColorIO/Source/OpenColorIOEditor/Private/OpenColorIOEditorBlueprintLibrary.cpp
@@ -20,6 +20,16 @@ void UOpenColorIOEditorBlueprintLibrary::SetActiveViewportConfiguration(const FO
 	OpenColorIOEditorModule.SetActiveViewportConfiguration(InConfiguration);
 }
 
+bool UOpenColorIOEditorBlueprintLibrary::ApplyColorSpaceTransformToColor(const FOpenColorIOColorConversionSettings& ConversionSettings, FLinearColor& InOutColor)
+{
+	if (IsValid(ConversionSettings.ConfigurationSource))
+	{
+		return ConversionSettings.ConfigurationSource->TransformColor(ConversionSettings, InOutColor);
+	}
+
+	return false;
+}
+
 bool UOpenColorIOEditorBlueprintLibrary::ApplyColorSpaceTransformToTexture(const FOpenColorIOColorConversionSettings& ConversionSettings, UTexture* InOutTexture)
 {
 	if (IsValid(InOutTexture))
diff --git a/Engine/Plugins/Compositing/OpenColorIO/Source/OpenColorIOEditor/Public/OpenColorIOEditorBlueprintLibrary.h b/Engine/Plugins/Compositing/OpenColorIO/Source/OpenColorIOEditor/Public/OpenColorIOEditorBlueprintLibrary.h
--- a/Engine/Plugins/Compositing/OpenColorIO/Source/OpenColorIOEditor/Public/OpenColorIOEditorBlueprintLibrary.h
+++ b/Engine/Plugins/Compositing/OpenColorIO/Source/OpenColorIOEditor/Public/OpenColorIOEditorBlueprintLibrary.h
@@ -23,6 +23,16 @@ public:
 	UFUNCTION(BlueprintCallable, Category = "OpenColorIO")
 	static OPENCOLORIOEDITOR_API void SetActiveViewportConfiguration(const FOpenColorIODisplayConfiguration& InConfiguration);
 
+	/**
+	 * Apply a color space transform to a color value.
+	 *
+	 * @param ConversionSettings Color transformation settings.
+	 * @param InOutColor Color value transform.
+	 * @return true upon success.
+	 */
+	UFUNCTION(BlueprintCallable, Category = "OpenColorIO")
+	static OPENCOLORIOEDITOR_API bool ApplyColorSpaceTransformToColor(const FOpenColorIOColorConversionSettings& ConversionSettings, FLinearColor& InOutColor);
+
 	/**
 	 * Apply a color space transform to a texture asset.
 	 *
