OCIO: Fix just-added blueprint interface for CPU color transform.

#rnx

diff --git a/Engine/Plugins/Compositing/OpenColorIO/Source/OpenColorIOEditor/Private/OpenColorIOEditorBlueprintLibrary.cpp b/Engine/Plugins/Compositing/OpenColorIO/Source/OpenColorIOEditor/Private/OpenColorIOEditorBlueprintLibrary.cpp
--- a/Engine/Plugins/Compositing/OpenColorIO/Source/OpenColorIOEditor/Private/OpenColorIOEditorBlueprintLibrary.cpp
+++ b/Engine/Plugins/Compositing/OpenColorIO/Source/OpenColorIOEditor/Private/OpenColorIOEditorBlueprintLibrary.cpp
@@ -20,11 +20,13 @@ void UOpenColorIOEditorBlueprintLibrary::SetActiveViewportConfiguration(const FO
 	OpenColorIOEditorModule.SetActiveViewportConfiguration(InConfiguration);
 }
 
-bool UOpenColorIOEditorBlueprintLibrary::ApplyColorSpaceTransformToColor(const FOpenColorIOColorConversionSettings& ConversionSettings, FLinearColor& InOutColor)
+bool UOpenColorIOEditorBlueprintLibrary::ApplyColorSpaceTransformToColor(const FOpenColorIOColorConversionSettings& ConversionSettings, const FLinearColor& InColor, FLinearColor& OutColor)
 {
 	if (IsValid(ConversionSettings.ConfigurationSource))
 	{
-		return ConversionSettings.ConfigurationSource->TransformColor(ConversionSettings, InOutColor);
+		OutColor = InColor;
+
+		return ConversionSettings.ConfigurationSource->TransformColor(ConversionSettings, OutColor);
 	}
 
 	return false;
diff --git a/Engine/Plugins/Compositing/OpenColorIO/Source/OpenColorIOEditor/Public/OpenColorIOEditorBlueprintLibrary.h b/Engine/Plugins/Compositing/OpenColorIO/Source/OpenColorIOEditor/Public/OpenColorIOEditorBlueprintLibrary.h
--- a/Engine/Plugins/Compositing/OpenColorIO/Source/OpenColorIOEditor/Public/OpenColorIOEditorBlueprintLibrary.h
+++ b/Engine/Plugins/Compositing/OpenColorIO/Source/OpenColorIOEditor/Public/OpenColorIOEditorBlueprintLibrary.h
@@ -31,7 +31,7 @@ public:
 	 * @return true upon success.
 	 */
 	UFUNCTION(BlueprintCallable, Category = "OpenColorIO")
-	static OPENCOLORIOEDITOR_API bool ApplyColorSpaceTransformToColor(const FOpenColorIOColorConversionSettings& ConversionSettings, FLinearColor& InOutColor);
+	static OPENCOLORIOEDITOR_API bool ApplyColorSpaceTransformToColor(const FOpenColorIOColorConversionSettings& ConversionSettings, const FLinearColor& InColor, FLinearColor& OutColor);
 
 	/**
 	 * Apply a color space transform to a texture asset.
