OCIO: General improvements before MRQ CPU integration.

diff --git a/Engine/Plugins/Compositing/OpenColorIO/Source/OpenColorIO/Private/OpenColorIOColorTransform.cpp b/Engine/Plugins/Compositing/OpenColorIO/Source/OpenColorIO/Private/OpenColorIOColorTransform.cpp
--- a/Engine/Plugins/Compositing/OpenColorIO/Source/OpenColorIO/Private/OpenColorIOColorTransform.cpp
+++ b/Engine/Plugins/Compositing/OpenColorIO/Source/OpenColorIO/Private/OpenColorIOColorTransform.cpp
@@ -37,33 +37,6 @@ namespace
 
 		return nullptr;
 	}
-
-	// Note: Since OpenColorIO caches processors automatically, we can recreate them without significant additional costs.
-	FOpenColorIOWrapperProcessor GetTransformProcessor(const UOpenColorIOColorTransform* InTransform, const FOpenColorIOWrapperConfig* InConfigWrapper)
-	{
-		check(InTransform);
-		ensure(InConfigWrapper);
-
-		EOpenColorIOViewTransformDirection DisplayViewDirection;
-		if (InTransform->GetDisplayViewDirection(DisplayViewDirection))
-		{
-			return FOpenColorIOWrapperProcessor(
-				InConfigWrapper,
-				InTransform->SourceColorSpace,
-				InTransform->Display,
-				InTransform->View,
-				DisplayViewDirection == EOpenColorIOViewTransformDirection::Inverse,
-				InTransform->GetContextKeyValues());
-		}
-		else
-		{
-			return FOpenColorIOWrapperProcessor(
-				InConfigWrapper,
-				InTransform->SourceColorSpace,
-				InTransform->DestinationColorSpace,
-				InTransform->GetContextKeyValues());
-		}
-	}
 }
 #endif //WITH_EDITOR
 
@@ -290,68 +263,56 @@ void UOpenColorIOColorTransform::CacheResourceShadersForCooking(EShaderPlatform
 
 void UOpenColorIOColorTransform::ProcessTransform()
 {
-	const UOpenColorIOConfiguration* ConfigurationOwner = Cast<UOpenColorIOConfiguration>(GetOuter());
-	const FOpenColorIOWrapperConfig* ConfigWrapper = GetTransformConfigWrapper(ConfigurationOwner);
-
-	if (ConfigWrapper != nullptr)
+	FOpenColorIOWrapperProcessor TransformProcessor;
+	if (GetTransformProcessor(TransformProcessor))
 	{
-		const FOpenColorIOWrapperProcessor TransformProcessor = GetTransformProcessor(this, ConfigWrapper);
-
-		if (TransformProcessor.IsValid())
-		{
-			Textures.Reset();
+		Textures.Reset();
 
-			const UOpenColorIOSettings* Settings = GetDefault<UOpenColorIOSettings>();
-			const FOpenColorIOWrapperGPUProcessor GPUProcessor = FOpenColorIOWrapperGPUProcessor(TransformProcessor, Settings->bUseLegacyProcessor);
-			const FString GpuProcessorHash = GPUProcessor.GetCacheID();
+		const UOpenColorIOSettings* Settings = GetDefault<UOpenColorIOSettings>();
+		const FOpenColorIOWrapperGPUProcessor GPUProcessor = FOpenColorIOWrapperGPUProcessor(TransformProcessor, Settings->bUseLegacyProcessor);
+		const FString GpuProcessorHash = GPUProcessor.GetCacheID();
 
-			// Process 3D luts
-			for (uint32 Index = 0; Index < GPUProcessor.GetNum3DTextures(); ++Index)
-			{
-				TextureFilter Filter;
-				FName TextureName;
-				uint32 EdgeLength = 0;
-				const float* TextureValues = 0x0;
-
-				bool bSuccess = GPUProcessor.Get3DTexture(Index, TextureName, EdgeLength, Filter, TextureValues);
-				check(bSuccess);
-
-				TObjectPtr<UTexture> Result = CreateTexture3DLUT(GpuProcessorHash, TextureName, EdgeLength, Filter, TextureValues);
+		// Process 3D luts
+		for (uint32 Index = 0; Index < GPUProcessor.GetNum3DTextures(); ++Index)
+		{
+			TextureFilter Filter;
+			FName TextureName;
+			uint32 EdgeLength = 0;
+			const float* TextureValues = 0x0;
 
-				const int32 SlotIndex = TextureName.GetNumber() - 1; // Rely on FName's index number extraction for convenience
-				Textures.Add(SlotIndex, MoveTemp(Result));
-			}
+			bool bSuccess = GPUProcessor.Get3DTexture(Index, TextureName, EdgeLength, Filter, TextureValues);
+			check(bSuccess);
 
-			// Process 1D luts
-			for (uint32 Index = 0; Index < GPUProcessor.GetNumTextures(); ++Index)
-			{
-				FName TextureName;
-				uint32 TextureWidth = 0;
-				uint32 TextureHeight = 0;
-				TextureFilter Filter;
-				bool bRedChannelOnly = false;
-				const float* TextureValues = 0x0;
+			TObjectPtr<UTexture> Result = CreateTexture3DLUT(GpuProcessorHash, TextureName, EdgeLength, Filter, TextureValues);
 
-				bool bSuccess = GPUProcessor.GetTexture(Index, TextureName, TextureWidth, TextureHeight, Filter, bRedChannelOnly, TextureValues);
-				checkf(bSuccess, TEXT("Failed to read OCIO 1D LUT data."));
+			const int32 SlotIndex = TextureName.GetNumber() - 1; // Rely on FName's index number extraction for convenience
+			Textures.Add(SlotIndex, MoveTemp(Result));
+		}
 
-				TObjectPtr<UTexture> Result = CreateTexture1DLUT(GpuProcessorHash, TextureName, TextureWidth, TextureHeight, Filter, bRedChannelOnly, TextureValues);
+		// Process 1D luts
+		for (uint32 Index = 0; Index < GPUProcessor.GetNumTextures(); ++Index)
+		{
+			FName TextureName;
+			uint32 TextureWidth = 0;
+			uint32 TextureHeight = 0;
+			TextureFilter Filter;
+			bool bRedChannelOnly = false;
+			const float* TextureValues = 0x0;
 
-				const int32 SlotIndex = TextureName.GetNumber() - 1; // Rely on FName's index number extraction for convenience
-				Textures.Add(SlotIndex, MoveTemp(Result));
-			}
+			bool bSuccess = GPUProcessor.GetTexture(Index, TextureName, TextureWidth, TextureHeight, Filter, bRedChannelOnly, TextureValues);
+			checkf(bSuccess, TEXT("Failed to read OCIO 1D LUT data."));
 
-			ensureAlwaysMsgf(Textures.Num() <= (int32)OpenColorIOShader::MaximumTextureSlots, TEXT("Color transform %s exceeds our current limit of %u texture slots. Use the legacy processor instead."), *GetTransformFriendlyName(), OpenColorIOShader::MaximumTextureSlots);
-			
-			// Generate shader code
-			GPUProcessor.GetShader(GeneratedShaderHash, GeneratedShader);
+			TObjectPtr<UTexture> Result = CreateTexture1DLUT(GpuProcessorHash, TextureName, TextureWidth, TextureHeight, Filter, bRedChannelOnly, TextureValues);
 
-			// Early return on success
-			return;
+			const int32 SlotIndex = TextureName.GetNumber() - 1; // Rely on FName's index number extraction for convenience
+			Textures.Add(SlotIndex, MoveTemp(Result));
 		}
-	}
 
-	UE_LOG(LogOpenColorIO, Error, TEXT("Failed to process and cache resource(s) for color transform %s. Configuration file [%s] was invalid."), *GetTransformFriendlyName(), ConfigurationOwner ? *ConfigurationOwner->ConfigurationFile.FilePath : TEXT("Missing"));
+		ensureAlwaysMsgf(Textures.Num() <= (int32)OpenColorIOShader::MaximumTextureSlots, TEXT("Color transform %s exceeds our current limit of %u texture slots. Use the legacy processor instead."), *GetTransformFriendlyName(), OpenColorIOShader::MaximumTextureSlots);
+		
+		// Generate shader code
+		GPUProcessor.GetShader(GeneratedShaderHash, GeneratedShader);
+	}
 }
 
 TObjectPtr<UTexture> UOpenColorIOColorTransform::CreateTexture3DLUT(const FString& InProcessorIdentifier, const FName& InName, uint32 InLutLength, TextureFilter InFilter, const float* InSourceData)
@@ -543,7 +504,7 @@ FOpenColorIOTransformResource* UOpenColorIOColorTransform::AllocateResource()
 	return new FOpenColorIOTransformResource();
 }
 
-bool UOpenColorIOColorTransform::GetRenderResources(ERHIFeatureLevel::Type InFeatureLevel, FOpenColorIOTransformResource*& OutShaderResource, TSortedMap<int32, FTextureResource*>& OutTextureResources)
+bool UOpenColorIOColorTransform::GetRenderResources(ERHIFeatureLevel::Type InFeatureLevel, FOpenColorIOTransformResource*& OutShaderResource, TSortedMap<int32, FTextureResource*>& OutTextureResources) const
 {
 	OutShaderResource = ColorTransformResources[InFeatureLevel];
 	
@@ -628,43 +589,48 @@ bool UOpenColorIOColorTransform::IsTransform(const FString& InSourceColorSpace,
 }
 
 #if WITH_EDITOR
-bool UOpenColorIOColorTransform::TransformColor(FLinearColor& InOutColor)
+bool UOpenColorIOColorTransform::GetTransformProcessor(FOpenColorIOWrapperProcessor& OutProcessor) const
 {
 	const UOpenColorIOConfiguration* ConfigurationOwner = Cast<UOpenColorIOConfiguration>(GetOuter());
 	const FOpenColorIOWrapperConfig* ConfigWrapper = GetTransformConfigWrapper(ConfigurationOwner);
 
-	if (ConfigWrapper != nullptr)
+	if (!ConfigWrapper)
 	{
-		const FOpenColorIOWrapperProcessor TransformProcessor = GetTransformProcessor(this, ConfigWrapper);
-		if (TransformProcessor.IsValid())
-		{
-			return TransformProcessor.TransformColor(InOutColor);
-		}
+		UE_LOG(LogOpenColorIO, Error, TEXT("Failed to create processor for color transform %s. Configuration file [%s] was invalid."), *GetTransformFriendlyName(), ConfigurationOwner ? *ConfigurationOwner->ConfigurationFile.FilePath : TEXT("Missing"));
+		return false;
+	}
+
+	// Note: Since OpenColorIO caches processors automatically, we can recreate them without significant additional costs.
+	EOpenColorIOViewTransformDirection CurrentDisplayViewDirection;
+	if (GetDisplayViewDirection(CurrentDisplayViewDirection))
+	{
+		OutProcessor = FOpenColorIOWrapperProcessor(ConfigWrapper, SourceColorSpace, Display, View, CurrentDisplayViewDirection == EOpenColorIOViewTransformDirection::Inverse, GetContextKeyValues());
 	}
 	else
 	{
-		UE_LOG(LogOpenColorIO, Error, TEXT("Failed to transform image for color transform %s. Configuration file was invalid."), *GetTransformFriendlyName());
+		OutProcessor = FOpenColorIOWrapperProcessor(ConfigWrapper, SourceColorSpace, DestinationColorSpace, GetContextKeyValues());
 	}
 
-	return false;
+	return OutProcessor.IsValid();
 }
 
-bool UOpenColorIOColorTransform::TransformImage(const FImageView& InOutImage) const
+bool UOpenColorIOColorTransform::TransformColor(FLinearColor& InOutColor) const
 {
-	const UOpenColorIOConfiguration* ConfigurationOwner = Cast<UOpenColorIOConfiguration>(GetOuter());
-	const FOpenColorIOWrapperConfig* ConfigWrapper = GetTransformConfigWrapper(ConfigurationOwner);
-
-	if (ConfigWrapper != nullptr)
+	FOpenColorIOWrapperProcessor TransformProcessor;
+	if (GetTransformProcessor(TransformProcessor))
 	{
-		const FOpenColorIOWrapperProcessor TransformProcessor = GetTransformProcessor(this, ConfigWrapper);
-		if (TransformProcessor.IsValid())
-		{
-			return TransformProcessor.TransformImage(InOutImage);
-		}
+		return TransformProcessor.TransformColor(InOutColor);
 	}
-	else
+
+	return false;
+}
+
+bool UOpenColorIOColorTransform::TransformImage(const FImageView& InOutImage) const
+{
+	FOpenColorIOWrapperProcessor TransformProcessor;
+	if (GetTransformProcessor(TransformProcessor))
 	{
-		UE_LOG(LogOpenColorIO, Error, TEXT("Failed to transform image for color transform %s. Configuration file was invalid."), *GetTransformFriendlyName());
+		return TransformProcessor.TransformImage(InOutImage);
 	}
 
 	return false;
@@ -672,20 +638,10 @@ bool UOpenColorIOColorTransform::TransformImage(const FImageView& InOutImage) co
 
 bool UOpenColorIOColorTransform::TransformImage(const FImageView& SrcImage, const FImageView& DestImage) const
 {
-	const UOpenColorIOConfiguration* ConfigurationOwner = Cast<UOpenColorIOConfiguration>(GetOuter());
-	const FOpenColorIOWrapperConfig* ConfigWrapper = GetTransformConfigWrapper(ConfigurationOwner);
-
-	if (ConfigWrapper != nullptr)
-	{
-		const FOpenColorIOWrapperProcessor TransformProcessor = GetTransformProcessor(this, ConfigWrapper);
-		if (TransformProcessor.IsValid())
-		{
-			return TransformProcessor.TransformImage(SrcImage, DestImage);
-		}
-	}
-	else
+	FOpenColorIOWrapperProcessor TransformProcessor;
+	if (GetTransformProcessor(TransformProcessor))
 	{
-		UE_LOG(LogOpenColorIO, Error, TEXT("Failed to transform image for color transform %s. Configuration file was invalid."), *GetTransformFriendlyName());
+		return TransformProcessor.TransformImage(SrcImage, DestImage);
 	}
 
 	return false;
@@ -850,20 +806,22 @@ bool UOpenColorIOColorTransform::UpdateShaderInfo(FString& OutShaderCodeHash, FS
 
 	if (ConfigWrapper != nullptr)
 	{
-		const FOpenColorIOWrapperProcessor Processor = GetTransformProcessor(this, ConfigWrapper);
-		if (Processor.IsValid())
+		FOpenColorIOWrapperProcessor TransformProcessor;
+		if (GetTransformProcessor(TransformProcessor))
 		{
-			const FOpenColorIOWrapperGPUProcessor GPUProcessor = FOpenColorIOWrapperGPUProcessor(Processor, Settings->bUseLegacyProcessor);
+			const FOpenColorIOWrapperGPUProcessor GPUProcessor = FOpenColorIOWrapperGPUProcessor(TransformProcessor, Settings->bUseLegacyProcessor);
 			if (GPUProcessor.IsValid())
 			{
-				OutRawConfigHash = ConfigurationOwner->GetConfigWrapper()->GetCacheID();
+				OutRawConfigHash = ConfigWrapper->GetCacheID();
 
 				return GPUProcessor.GetShader(OutShaderCodeHash, OutShaderCode);
 			}
 		}
 	}
-
-	UE_LOG(LogOpenColorIO, Error, TEXT("Failed to fetch shader info for color transform %s. Configuration file [%s] was invalid."), *GetTransformFriendlyName(), ConfigurationOwner ? *ConfigurationOwner->ConfigurationFile.FilePath : TEXT("Missing"));
+	else
+	{
+		UE_LOG(LogOpenColorIO, Error, TEXT("Failed to fetch shader info for color transform %s. Configuration file [%s] was invalid."), *GetTransformFriendlyName(), ConfigurationOwner ? *ConfigurationOwner->ConfigurationFile.FilePath : TEXT("Missing"));
+	}
 
 	return false;
 }
diff --git a/Engine/Plugins/Compositing/OpenColorIO/Source/OpenColorIO/Private/OpenColorIOConfiguration.cpp b/Engine/Plugins/Compositing/OpenColorIO/Source/OpenColorIO/Private/OpenColorIOConfiguration.cpp
--- a/Engine/Plugins/Compositing/OpenColorIO/Source/OpenColorIO/Private/OpenColorIOConfiguration.cpp
+++ b/Engine/Plugins/Compositing/OpenColorIO/Source/OpenColorIO/Private/OpenColorIOConfiguration.cpp
@@ -60,9 +60,10 @@ void UOpenColorIOConfiguration::BeginDestroy()
 
 bool UOpenColorIOConfiguration::IsTransformReady(const FOpenColorIOColorConversionSettings& InSettings)
 {
-	if (const TObjectPtr<UOpenColorIOColorTransform>* TransformPtr = FindTransform(InSettings))
+	TObjectPtr<const UOpenColorIOColorTransform> Transform = FindTransform(InSettings);
+	if (IsValid(Transform))
 	{
-		return (*TransformPtr)->AreRenderResourcesReady();
+		return Transform->AreRenderResourcesReady();
 	}
 
 	return false;
@@ -70,9 +71,10 @@ bool UOpenColorIOConfiguration::IsTransformReady(const FOpenColorIOColorConversi
 
 bool UOpenColorIOConfiguration::GetRenderResources(ERHIFeatureLevel::Type InFeatureLevel, const FOpenColorIOColorConversionSettings& InSettings, FOpenColorIOTransformResource*& OutShaderResource, TSortedMap<int32, FTextureResource*>& OutTextureResources)
 {
-	if (const TObjectPtr<UOpenColorIOColorTransform>* TransformPtr = FindTransform(InSettings))
+	TObjectPtr<const UOpenColorIOColorTransform> Transform = FindTransform(InSettings);
+	if (IsValid(Transform))
 	{
-		return (*TransformPtr)->GetRenderResources(InFeatureLevel, OutShaderResource, OutTextureResources);
+		return Transform->GetRenderResources(InFeatureLevel, OutShaderResource, OutTextureResources);
 	}
 
 	return false;
@@ -136,9 +138,10 @@ bool UOpenColorIOConfiguration::Validate() const
 #if WITH_EDITOR
 bool UOpenColorIOConfiguration::TransformColor(const FOpenColorIOColorConversionSettings& InSettings, FLinearColor& InOutColor) const
 {
-	if (const TObjectPtr<UOpenColorIOColorTransform>* TransformPtr = FindTransform(InSettings))
+	TObjectPtr<const UOpenColorIOColorTransform> Transform = FindTransform(InSettings);
+	if (IsValid(Transform))
 	{
-		return (*TransformPtr)->TransformColor(InOutColor);
+		return Transform->TransformColor(InOutColor);
 	}
 
 	return false;
@@ -146,9 +149,10 @@ bool UOpenColorIOConfiguration::TransformColor(const FOpenColorIOColorConversion
 
 bool UOpenColorIOConfiguration::TransformImage(const FOpenColorIOColorConversionSettings& InSettings, const FImageView& InOutImage) const
 {
-	if (const TObjectPtr<UOpenColorIOColorTransform>* TransformPtr = FindTransform(InSettings))
+	TObjectPtr<const UOpenColorIOColorTransform> Transform = FindTransform(InSettings);
+	if (IsValid(Transform))
 	{
-		return (*TransformPtr)->TransformImage(InOutImage);
+		return Transform->TransformImage(InOutImage);
 	}
 
 	return false;
@@ -156,16 +160,17 @@ bool UOpenColorIOConfiguration::TransformImage(const FOpenColorIOColorConversion
 
 bool UOpenColorIOConfiguration::TransformImage(const FOpenColorIOColorConversionSettings& InSettings, const FImageView& SrcImage, const FImageView& DestImage) const
 {
-	if (const TObjectPtr<UOpenColorIOColorTransform>* TransformPtr = FindTransform(InSettings))
+	TObjectPtr<const UOpenColorIOColorTransform> Transform = FindTransform(InSettings);
+	if (IsValid(Transform))
 	{
-		return (*TransformPtr)->TransformImage(SrcImage, DestImage);
+		return Transform->TransformImage(SrcImage, DestImage);
 	}
 
 	return false;
 }
 #endif
 
-void UOpenColorIOConfiguration::ReloadExistingColorspaces()
+void UOpenColorIOConfiguration::ReloadExistingColorspaces(bool bForce)
 {
 #if WITH_EDITOR
 	LoadConfiguration();
@@ -181,7 +186,7 @@ void UOpenColorIOConfiguration::ReloadExistingColorspaces()
 		}
 
 		// Hash is different, proceed with the regeneration...
-		if (ConfigHash != LoadedConfigHash)
+		if (ConfigHash != LoadedConfigHash || bForce)
 		{
 			ConfigHash = LoadedConfigHash;
 
@@ -369,9 +374,9 @@ void UOpenColorIOConfiguration::StopDirectoryWatch()
 #endif
 }
 
-const TObjectPtr<UOpenColorIOColorTransform>* UOpenColorIOConfiguration::FindTransform(const FOpenColorIOColorConversionSettings& InSettings) const
+TObjectPtr<const UOpenColorIOColorTransform> UOpenColorIOConfiguration::FindTransform(const FOpenColorIOColorConversionSettings& InSettings) const
 {
-	return ColorTransforms.FindByPredicate([&](const UOpenColorIOColorTransform* InTransform)
+	const TObjectPtr<UOpenColorIOColorTransform>* TransformPtr = ColorTransforms.FindByPredicate([&](const UOpenColorIOColorTransform* InTransform)
 		{
 			EOpenColorIOViewTransformDirection DisplayViewDirection;
 			
@@ -384,6 +389,8 @@ const TObjectPtr<UOpenColorIOColorTransform>* UOpenColorIOConfiguration::FindTra
 				return InTransform->SourceColorSpace == InSettings.SourceColorSpace.ColorSpaceName && InTransform->DestinationColorSpace == InSettings.DestinationColorSpace.ColorSpaceName;
 			}
 		});
+
+	return (TransformPtr != nullptr) ? *TransformPtr : nullptr;
 }
 
 #if WITH_EDITOR
diff --git a/Engine/Plugins/Compositing/OpenColorIO/Source/OpenColorIO/Public/OpenColorIOColorTransform.h b/Engine/Plugins/Compositing/OpenColorIO/Source/OpenColorIO/Public/OpenColorIOColorTransform.h
--- a/Engine/Plugins/Compositing/OpenColorIO/Source/OpenColorIO/Public/OpenColorIOColorTransform.h
+++ b/Engine/Plugins/Compositing/OpenColorIO/Source/OpenColorIO/Public/OpenColorIOColorTransform.h
@@ -16,6 +16,7 @@
 
 
 struct FImageView;
+class FOpenColorIOWrapperProcessor;
 class UOpenColorIOConfiguration;
 
 
@@ -69,19 +70,33 @@ public:
 	/**
 	 * Returns the desired resources required to apply this transform during rendering.
 	 */
-	bool GetRenderResources(ERHIFeatureLevel::Type InFeatureLevel, FOpenColorIOTransformResource*& OutShaderResource, TSortedMap<int32, FTextureResource*>& OutTextureResources);
+	bool GetRenderResources(ERHIFeatureLevel::Type InFeatureLevel, FOpenColorIOTransformResource*& OutShaderResource, TSortedMap<int32, FTextureResource*>& OutTextureResources) const;
 
 	/**
 	 * Returns true if shader/texture resources have finished compiling and are ready for use (to be called on the game thread).
 	 */
 	bool AreRenderResourcesReady() const;
 
+	/**
+	 * Returns true if the current transform corresponds to the specified color spaces.
+	 */
 	bool IsTransform(const FString& InSourceColorSpace, const FString& InDestinationColorSpace) const;
+
+	/**
+	 * Returns true if the current transform corresponds to the specified color space and display/view & direction.
+	 */
 	bool IsTransform(const FString& InSourceColorSpace, const FString& InDisplay, const FString& InView, EOpenColorIOViewTransformDirection InDirection) const;
 
 #if WITH_EDITOR
+	/**
+	 * Get the transform processor.
+	 *
+	 * @return True if the processor was created and is valid.
+	 */
+	bool GetTransformProcessor(FOpenColorIOWrapperProcessor& OutProcessor) const;
+
 	/** Apply the color transform in-place to the specified color. */
-	bool TransformColor(FLinearColor& InOutColor);
+	bool TransformColor(FLinearColor& InOutColor) const;
 		
 	/** Apply the color transform in-place to the specified image. */
 	bool TransformImage(const FImageView& InOutImage) const;
diff --git a/Engine/Plugins/Compositing/OpenColorIO/Source/OpenColorIO/Public/OpenColorIOConfiguration.h b/Engine/Plugins/Compositing/OpenColorIO/Source/OpenColorIO/Public/OpenColorIOConfiguration.h
--- a/Engine/Plugins/Compositing/OpenColorIO/Source/OpenColorIO/Public/OpenColorIOConfiguration.h
+++ b/Engine/Plugins/Compositing/OpenColorIO/Source/OpenColorIO/Public/OpenColorIOConfiguration.h
@@ -58,7 +58,7 @@ public:
 
 	/** This forces to reload colorspaces and corresponding shaders if those are not loaded already. */
 	UFUNCTION(BlueprintCallable, Category = "OpenColorIO")
-	void ReloadExistingColorspaces();
+	void ReloadExistingColorspaces(bool bForce = false);
 
 	/*
 	* This method is called by directory watcher when any file or folder is changed in the 
@@ -68,10 +68,12 @@ public:
 
 	/** Internal only: Replacement for previous `GetNativeConfig_Internal()`, `GetLoadedConfiguration()` and `GetLoadedConfigurationFile()` functions, returning the private implementation of the native OCIO config. */
 	FOpenColorIOWrapperConfig* GetConfigWrapper() const;
+	
+	/** Find the color transform object that corresponds to the specified settings, nullptr if not found. */
+	TObjectPtr<const UOpenColorIOColorTransform> FindTransform(const FOpenColorIOColorConversionSettings& InSettings) const;
 
 protected:
 
-	const TObjectPtr<UOpenColorIOColorTransform>* FindTransform(const FOpenColorIOColorConversionSettings& InSettings) const;
 #if WITH_EDITOR
 	void CreateColorTransform(const FString& InSourceColorSpace, const FString& InDestinationColorSpace);
 	void CreateColorTransform(const FString& InSourceColorSpace, const FString& InDisplay, const FString& InView, EOpenColorIOViewTransformDirection InDirection);
diff --git a/Engine/Source/Developer/OpenColorIOWrapper/Private/OpenColorIOWrapper.cpp b/Engine/Source/Developer/OpenColorIOWrapper/Private/OpenColorIOWrapper.cpp
--- a/Engine/Source/Developer/OpenColorIOWrapper/Private/OpenColorIOWrapper.cpp
+++ b/Engine/Source/Developer/OpenColorIOWrapper/Private/OpenColorIOWrapper.cpp
@@ -68,6 +68,7 @@ struct FOpenColorIOConfigPimpl
 struct FOpenColorIOProcessorPimpl
 {
 	OCIO_NAMESPACE::ConstProcessorRcPtr Processor = nullptr;
+	OCIO_NAMESPACE::ConstProcessorRcPtr InterchangeProcessor = nullptr;
 
 	/** Get processor optimization flags. */
 	static OCIO_NAMESPACE::OptimizationFlags GetOptimizationFlags()
@@ -77,6 +78,29 @@ struct FOpenColorIOProcessorPimpl
 			OCIO_NAMESPACE::OptimizationFlags::OPTIMIZATION_NO_DYNAMIC_PROPERTIES
 		);
 	}
+
+	/** Create the interchange processor used for conversion between the working color space and the interchange space. */
+	void SetupInterchangeProcessor(const OCIO_NAMESPACE::ConstConfigRcPtr& InConfig, const OCIO_NAMESPACE::ConstConfigRcPtr& InterchangeConfig, EOpenColorIOWorkingColorSpaceTransform InWorkingColorSpaceTransformType)
+	{
+		const TUniquePtr<ANSICHAR[]> AnsiWorkingColorSpaceName = OpenColorIOWrapper::MakeAnsiString(OpenColorIOWrapper::GetWorkingColorSpaceName());
+
+		if (InWorkingColorSpaceTransformType == EOpenColorIOWorkingColorSpaceTransform::Source)
+		{
+			InterchangeProcessor = InterchangeConfig->GetProcessorFromConfigs(
+				InterchangeConfig,
+				AnsiWorkingColorSpaceName.Get(),
+				InConfig,
+				InConfig->getCanonicalName(OCIO_NAMESPACE::ROLE_INTERCHANGE_SCENE));
+		}
+		else if (InWorkingColorSpaceTransformType == EOpenColorIOWorkingColorSpaceTransform::Destination)
+		{
+			InterchangeProcessor = InterchangeConfig->GetProcessorFromConfigs(
+				InConfig,
+				InConfig->getCanonicalName(OCIO_NAMESPACE::ROLE_INTERCHANGE_SCENE),
+				InterchangeConfig,
+				AnsiWorkingColorSpaceName.Get());
+		}
+	}
 };
 
 struct FOpenColorIOGPUProcessorPimpl
@@ -388,13 +412,16 @@ FOpenColorIOWrapperEngineConfig::FOpenColorIOWrapperEngineConfig()
 	check(IsValid());
 }
 
+FOpenColorIOWrapperProcessor::FOpenColorIOWrapperProcessor()
+	: Pimpl(MakePimpl<FOpenColorIOProcessorPimpl, EPimplPtrMode::DeepCopy>())
+{ }
+
 FOpenColorIOWrapperProcessor::FOpenColorIOWrapperProcessor(
 	const FOpenColorIOWrapperConfig* InConfig,
 	FStringView InSourceColorSpace,
 	FStringView InDestinationColorSpace,
 	const TMap<FString, FString>& InContextKeyValues)
 	: Pimpl(MakePimpl<FOpenColorIOProcessorPimpl, EPimplPtrMode::DeepCopy>())
-	, OwnerConfig(InConfig)
 	, WorkingColorSpaceTransformType(EOpenColorIOWorkingColorSpaceTransform::None)
 {
 	if (InSourceColorSpace == OpenColorIOWrapper::GetWorkingColorSpaceName())
@@ -407,9 +434,9 @@ FOpenColorIOWrapperProcessor::FOpenColorIOWrapperProcessor(
 	}
 
 	OCIO_EXCEPTION_HANDLING_TRY();
-		if (OwnerConfig != nullptr && OwnerConfig->IsValid())
+		if (InConfig != nullptr && InConfig->IsValid())
 		{
-			const OCIO_NAMESPACE::ConstConfigRcPtr& Config = OwnerConfig->Pimpl->Config;
+			const OCIO_NAMESPACE::ConstConfigRcPtr& Config = InConfig->Pimpl->Config;
 			OCIO_NAMESPACE::ContextRcPtr Context = Config->getCurrentContext()->createEditableCopy();
 
 			for (const TPair<FString, FString>& KeyValue : InContextKeyValues)
@@ -422,6 +449,11 @@ FOpenColorIOWrapperProcessor::FOpenColorIOWrapperProcessor(
 				StringCast<ANSICHAR>(InSourceColorSpace.GetData()).Get(),
 				StringCast<ANSICHAR>(InDestinationColorSpace.GetData()).Get()
 			);
+
+			if (WorkingColorSpaceTransformType != EOpenColorIOWorkingColorSpaceTransform::None)
+			{
+				Pimpl->SetupInterchangeProcessor(Config, IOpenColorIOWrapperModule::Get().GetEngineBuiltInConfig().Pimpl->Config, WorkingColorSpaceTransformType);
+			}
 		}
 	OCIO_EXCEPTION_HANDLING_CATCH(Log, TEXT("Failed to create processor for [%s, %s]. Error message: %s"), InSourceColorSpace.GetData(), InDestinationColorSpace.GetData());
 }
@@ -434,7 +466,6 @@ FOpenColorIOWrapperProcessor::FOpenColorIOWrapperProcessor(
 	bool bInverseDirection,
 	const TMap<FString, FString>& InContextKeyValues)
 	: Pimpl(MakePimpl<FOpenColorIOProcessorPimpl, EPimplPtrMode::DeepCopy>())
-	, OwnerConfig(InConfig)
 	, WorkingColorSpaceTransformType(EOpenColorIOWorkingColorSpaceTransform::None)
 {
 	if (InSourceColorSpace == OpenColorIOWrapper::GetWorkingColorSpaceName())
@@ -450,9 +481,9 @@ FOpenColorIOWrapperProcessor::FOpenColorIOWrapperProcessor(
 	}
 
 	OCIO_EXCEPTION_HANDLING_TRY();
-		if (OwnerConfig != nullptr && OwnerConfig->IsValid())
+		if (InConfig != nullptr && InConfig->IsValid())
 		{
-			const OCIO_NAMESPACE::ConstConfigRcPtr& Config = OwnerConfig->Pimpl->Config;
+			const OCIO_NAMESPACE::ConstConfigRcPtr& Config = InConfig->Pimpl->Config;
 			OCIO_NAMESPACE::ContextRcPtr Context = Config->getCurrentContext()->createEditableCopy();
 
 			for (const TPair<FString, FString>& KeyValue : InContextKeyValues)
@@ -466,19 +497,23 @@ FOpenColorIOWrapperProcessor::FOpenColorIOWrapperProcessor(
 				StringCast<ANSICHAR>(InView.GetData()).Get(),
 				static_cast<OCIO_NAMESPACE::TransformDirection>(bInverseDirection)
 			);
+
+			if (WorkingColorSpaceTransformType != EOpenColorIOWorkingColorSpaceTransform::None)
+			{
+				Pimpl->SetupInterchangeProcessor(Config, IOpenColorIOWrapperModule::Get().GetEngineBuiltInConfig().Pimpl->Config, WorkingColorSpaceTransformType);
+			}
 		}
 	OCIO_EXCEPTION_HANDLING_CATCH(Log, TEXT("Failed to create processor for [%s, %s, %s, %s]. Error message: %s"), InSourceColorSpace.GetData(), InDisplay.GetData(), InView.GetData(), (bInverseDirection ? TEXT("Inverse") : TEXT("Forward")));
 }
 
 FOpenColorIOWrapperProcessor::FOpenColorIOWrapperProcessor(const FOpenColorIOWrapperConfig* InConfig, FStringView InNamedTransform, bool bInverseDirection, const TMap<FString, FString>& InContextKeyValues)
 	: Pimpl(MakePimpl<FOpenColorIOProcessorPimpl, EPimplPtrMode::DeepCopy>())
-	, OwnerConfig(InConfig)
 	, WorkingColorSpaceTransformType(EOpenColorIOWorkingColorSpaceTransform::None)
 {
 	OCIO_EXCEPTION_HANDLING_TRY();
-		if (OwnerConfig != nullptr && OwnerConfig->IsValid())
+		if (InConfig != nullptr && InConfig->IsValid())
 		{
-			const OCIO_NAMESPACE::ConstConfigRcPtr& Config = OwnerConfig->Pimpl->Config;
+			const OCIO_NAMESPACE::ConstConfigRcPtr& Config = InConfig->Pimpl->Config;
 			OCIO_NAMESPACE::ContextRcPtr Context = Config->getCurrentContext()->createEditableCopy();
 
 			for (const TPair<FString, FString>& KeyValue : InContextKeyValues)
@@ -835,7 +870,7 @@ FOpenColorIOWrapperProcessor FOpenColorIOWrapperProcessor::CreateTransformToWork
 
 bool FOpenColorIOWrapperProcessor::IsValid() const
 {
-	return OwnerConfig != nullptr && Pimpl->Processor != nullptr;
+	return Pimpl->Processor != nullptr;
 }
 
 FString FOpenColorIOWrapperProcessor::GetCacheID() const
@@ -901,33 +936,12 @@ bool FOpenColorIOWrapperProcessor::TransformImage(const FImageView& InOutImage)
 
 				// Note: This special "interchange to working color space" logic should be removed
 				// once these features are supported by the library.
-				ConstProcessorRcPtr InterchangeProcessor = nullptr;
 				ConstCPUProcessorRcPtr InterchangeCPUProcessor = nullptr;
+				if (WorkingColorSpaceTransformType != EOpenColorIOWorkingColorSpaceTransform::None)
 				{
-					const TUniquePtr<ANSICHAR[]> AnsiWorkingColorSpaceName = OpenColorIOWrapper::MakeAnsiString(OpenColorIOWrapper::GetWorkingColorSpaceName());
-					ConstConfigRcPtr InterchangeConfig = IOpenColorIOWrapperModule::Get().GetEngineBuiltInConfig().Pimpl->Config;
-					ConstConfigRcPtr Config = OwnerConfig->Pimpl->Config;
-
-					if (WorkingColorSpaceTransformType == EOpenColorIOWorkingColorSpaceTransform::Source)
-					{
-						InterchangeProcessor = InterchangeConfig->GetProcessorFromConfigs(
-							InterchangeConfig,
-							AnsiWorkingColorSpaceName.Get(),
-							Config,
-							Config->getCanonicalName(OCIO_NAMESPACE::ROLE_INTERCHANGE_SCENE));
-					}
-					else if (WorkingColorSpaceTransformType == EOpenColorIOWorkingColorSpaceTransform::Destination)
+					if (ensure(Pimpl->InterchangeProcessor))
 					{
-						InterchangeProcessor = InterchangeConfig->GetProcessorFromConfigs(
-							Config,
-							Config->getCanonicalName(OCIO_NAMESPACE::ROLE_INTERCHANGE_SCENE),
-							InterchangeConfig,
-							AnsiWorkingColorSpaceName.Get());
-					}
-
-					if (InterchangeProcessor)
-					{
-						InterchangeCPUProcessor = InterchangeProcessor->getOptimizedCPUProcessor(BitDepth, BitDepth, OPTIMIZATION_DEFAULT);
+						InterchangeCPUProcessor = Pimpl->InterchangeProcessor->getOptimizedCPUProcessor(BitDepth, BitDepth, OPTIMIZATION_DEFAULT);
 					}
 				}
 
@@ -976,28 +990,19 @@ bool FOpenColorIOWrapperProcessor::TransformImage(const FImageView& SrcImage, co
 			{
 				PackedImageDesc SrcImageDesc = GetImageDesc(SrcImage);
 				PackedImageDesc DestImageDesc = GetImageDesc(DestImage);
-
-				const TUniquePtr<ANSICHAR[]> AnsiWorkingColorSpaceName = OpenColorIOWrapper::MakeAnsiString(OpenColorIOWrapper::GetWorkingColorSpaceName());
-				ConstConfigRcPtr	InterchangeConfig = IOpenColorIOWrapperModule::Get().GetEngineBuiltInConfig().Pimpl->Config;
-				ConstConfigRcPtr	Config = OwnerConfig->Pimpl->Config;
-
 				BitDepth SrcBitDepth = SrcImageDesc.getBitDepth();
 				BitDepth DestBitDepth = DestImageDesc.getBitDepth();
 
 				// Conditionally apply a conversion from the working color space to interchange space
 				if (WorkingColorSpaceTransformType == EOpenColorIOWorkingColorSpaceTransform::Source)
 				{
-					ConstProcessorRcPtr	InterchangeProcessor = InterchangeConfig->GetProcessorFromConfigs(
-						InterchangeConfig,
-						AnsiWorkingColorSpaceName.Get(),
-						Config,
-						Config->getCanonicalName(OCIO_NAMESPACE::ROLE_INTERCHANGE_SCENE));
-
-					ConstCPUProcessorRcPtr InterchangeCPUProcessor = InterchangeProcessor->getOptimizedCPUProcessor(SrcBitDepth, SrcBitDepth, OPTIMIZATION_DEFAULT);
-					InterchangeCPUProcessor->apply(SrcImageDesc);
+					if (ensure(Pimpl->InterchangeProcessor))
+					{
+						ConstCPUProcessorRcPtr InterchangeCPUProcessor = Pimpl->InterchangeProcessor->getOptimizedCPUProcessor(SrcBitDepth, SrcBitDepth, OPTIMIZATION_DEFAULT);
+						InterchangeCPUProcessor->apply(SrcImageDesc);
+					}
 				}
 
-
 				// Apply the main color transformation
 				ConstCPUProcessorRcPtr CPUProcessor = Pimpl->Processor->getOptimizedCPUProcessor(SrcBitDepth, DestBitDepth, OPTIMIZATION_DEFAULT);
 				CPUProcessor->apply(SrcImageDesc, DestImageDesc);
@@ -1005,14 +1010,11 @@ bool FOpenColorIOWrapperProcessor::TransformImage(const FImageView& SrcImage, co
 				// Conditionally apply a conversion from the interchange space to the working color space
 				if (WorkingColorSpaceTransformType == EOpenColorIOWorkingColorSpaceTransform::Destination)
 				{
-					ConstProcessorRcPtr	InterchangeProcessor = InterchangeConfig->GetProcessorFromConfigs(
-						Config,
-						Config->getCanonicalName(OCIO_NAMESPACE::ROLE_INTERCHANGE_SCENE),
-						InterchangeConfig,
-						AnsiWorkingColorSpaceName.Get());
-
-					ConstCPUProcessorRcPtr InterchangeCPUProcessor = InterchangeProcessor->getOptimizedCPUProcessor(DestBitDepth, DestBitDepth, OPTIMIZATION_DEFAULT);
-					InterchangeCPUProcessor->apply(DestImageDesc);
+					if (ensure(Pimpl->InterchangeProcessor))
+					{
+						ConstCPUProcessorRcPtr InterchangeCPUProcessor = Pimpl->InterchangeProcessor->getOptimizedCPUProcessor(DestBitDepth, DestBitDepth, OPTIMIZATION_DEFAULT);
+						InterchangeCPUProcessor->apply(DestImageDesc);
+					}
 				}
 
 				return true;
diff --git a/Engine/Source/Developer/OpenColorIOWrapper/Public/OpenColorIOWrapper.h b/Engine/Source/Developer/OpenColorIOWrapper/Public/OpenColorIOWrapper.h
--- a/Engine/Source/Developer/OpenColorIOWrapper/Public/OpenColorIOWrapper.h
+++ b/Engine/Source/Developer/OpenColorIOWrapper/Public/OpenColorIOWrapper.h
@@ -155,6 +155,12 @@ public:
 class FOpenColorIOWrapperProcessor final
 {
 public:
+
+	/**
+	* Empty constructor.
+	*/
+	OPENCOLORIOWRAPPER_API FOpenColorIOWrapperProcessor();
+
 	/**
 	* Constructor.
 	*
@@ -243,8 +249,6 @@ private:
 
 	TPimplPtr<struct FOpenColorIOProcessorPimpl, EPimplPtrMode::DeepCopy> Pimpl;
 
-	const FOpenColorIOWrapperConfig* OwnerConfig;
-
 	EOpenColorIOWorkingColorSpaceTransform WorkingColorSpaceTransformType;
 
 	friend class FOpenColorIOWrapperGPUProcessor;
