OCIO: Support $OCIO or %OCIO% environment variables in config paths, per client request.

#rnx

diff --git a/Engine/Plugins/Compositing/OpenColorIO/Source/OpenColorIO/Private/OpenColorIOConfiguration.cpp b/Engine/Plugins/Compositing/OpenColorIO/Source/OpenColorIO/Private/OpenColorIOConfiguration.cpp
--- a/Engine/Plugins/Compositing/OpenColorIO/Source/OpenColorIO/Private/OpenColorIOConfiguration.cpp
+++ b/Engine/Plugins/Compositing/OpenColorIO/Source/OpenColorIO/Private/OpenColorIOConfiguration.cpp
@@ -656,6 +656,10 @@ void UOpenColorIOConfiguration::LoadConfiguration()
 	{
 		bIsBuiltIn = true;
 	}
+	else if (ConfigurationFile.FilePath.Equals(TEXT("$OCIO")) || ConfigurationFile.FilePath.Equals(TEXT("%OCIO%")))
+	{
+		FilePath = FPlatformMisc::GetEnvironmentVariable(TEXT("OCIO"));
+	}
 	else if(ConfigurationFile.FilePath.Contains(TEXT("{Engine}")))
 	{
 		FilePath = FPaths::ConvertRelativePathToFull(ConfigurationFile.FilePath.Replace(TEXT("{Engine}"), *FPaths::EngineDir()));
