OCIO: Reload and Rebuild button should force the operation, even if the config hash is the same.

diff --git a/Engine/Plugins/Compositing/OpenColorIO/Source/OpenColorIOEditor/Private/OpenColorIOConfigurationCustomization.cpp b/Engine/Plugins/Compositing/OpenColorIO/Source/OpenColorIOEditor/Private/OpenColorIOConfigurationCustomization.cpp
--- a/Engine/Plugins/Compositing/OpenColorIO/Source/OpenColorIOEditor/Private/OpenColorIOConfigurationCustomization.cpp
+++ b/Engine/Plugins/Compositing/OpenColorIO/Source/OpenColorIOEditor/Private/OpenColorIOConfigurationCustomization.cpp
@@ -49,7 +49,7 @@ void FOpenColorIOConfigurationCustomization::CustomizeDetails(IDetailLayoutBuild
 				{
 					OpenColorIOWrapper::ClearAllCaches();
 
-					OpenColorIOConfig->ReloadExistingColorspaces();
+					OpenColorIOConfig->ReloadExistingColorspaces(true);
 				}
 				return FReply::Handled();
 			})
