OCIO: Improvements to SVE blueprint wrapper.

* Disabled SVE configuration before destroying it, to prevent active SVEs from holding references alive.
* Allow SVE creation with invalid transforms in game mode, given that we now support the fallback render pass.

diff --git a/Engine/Plugins/Compositing/OpenColorIO/Source/OpenColorIO/Private/OpenColorIODisplayExtensionWrapper.cpp b/Engine/Plugins/Compositing/OpenColorIO/Source/OpenColorIO/Private/OpenColorIODisplayExtensionWrapper.cpp
--- a/Engine/Plugins/Compositing/OpenColorIO/Source/OpenColorIO/Private/OpenColorIODisplayExtensionWrapper.cpp
+++ b/Engine/Plugins/Compositing/OpenColorIO/Source/OpenColorIO/Private/OpenColorIODisplayExtensionWrapper.cpp
@@ -64,6 +64,12 @@ void UOpenColorIODisplayExtensionWrapper::SetSceneExtensionIsActiveFunctions(con
 
 void UOpenColorIODisplayExtensionWrapper::RemoveSceneExtension()
 {
+	// We disable the display configuration before deletion so it's not accidentally kept alive by the GatherActiveExtensions() shared pointers.
+	if (DisplayExtension.IsValid())
+	{
+		DisplayExtension->GetDisplayConfiguration().bIsEnabled = false;
+	}
+
 	DisplayExtension.Reset();
 }
 
@@ -71,12 +77,6 @@ UOpenColorIODisplayExtensionWrapper* UOpenColorIODisplayExtensionWrapper::Create
 	FOpenColorIODisplayConfiguration InDisplayConfiguration,
 	const FSceneViewExtensionIsActiveFunctor& IsActiveFunction)
 {
-	if (!InDisplayConfiguration.ColorConfiguration.IsValid())
-	{
-		UE_LOG(LogOpenColorIO, Warning, TEXT("%s, no display extension was created."), *InDisplayConfiguration.ColorConfiguration.ToString());
-		return nullptr;
-	}
-
 	// Create OCIO Scene View Extension and configure it.
 
 	UOpenColorIODisplayExtensionWrapper* OutExtension = NewObject<UOpenColorIODisplayExtensionWrapper>();
