OCIO: Facilitate inspection of generated shaders.

diff --git a/Engine/Plugins/Compositing/OpenColorIO/Shaders/Private/OpenColorIOShader.usf b/Engine/Plugins/Compositing/OpenColorIO/Shaders/Private/OpenColorIOShader.usf
--- a/Engine/Plugins/Compositing/OpenColorIO/Shaders/Private/OpenColorIOShader.usf
+++ b/Engine/Plugins/Compositing/OpenColorIO/Shaders/Private/OpenColorIOShader.usf
@@ -29,11 +29,7 @@ uint TransformAlpha;
 
 float SafeRcp(float x)
 {
-	if (x > 0)
-	{
-		return rcp(x);
-	}
-	return 0;
+	return x > 0.0 ? rcp(x) : 0.0;
 }
 
 float4 InvertAlpha(float4 InColor)
diff --git a/Engine/Plugins/Compositing/OpenColorIO/Source/OpenColorIO/Private/OpenColorIOShaderMap.cpp b/Engine/Plugins/Compositing/OpenColorIO/Source/OpenColorIO/Private/OpenColorIOShaderMap.cpp
--- a/Engine/Plugins/Compositing/OpenColorIO/Source/OpenColorIO/Private/OpenColorIOShaderMap.cpp
+++ b/Engine/Plugins/Compositing/OpenColorIO/Source/OpenColorIO/Private/OpenColorIOShaderMap.cpp
@@ -38,6 +38,12 @@ namespace OpenColorIOShaderCookStats
 #endif
 
 
+static TAutoConsoleVariable<bool> CVarShowOpenColorIONoShaderDDC(
+	TEXT("OpenColorIO.NoShaderDDC"),
+	false,
+	TEXT("When enabled, we don't check the DDC for compiled shaders to simulate a cold DDC.")
+);
+
 //
 // Globals
 //
@@ -253,6 +259,13 @@ void FOpenColorIOShaderMap::LoadFromDerivedDataCache(const FOpenColorIOTransform
 	}
 	else
 	{
+		static bool bNoShaderDDC = FParse::Param(FCommandLine::Get(), TEXT("noshaderddc"));
+
+		if (bNoShaderDDC || CVarShowOpenColorIONoShaderDDC.GetValueOnAnyThread())
+		{
+			return;
+		}
+
 		// Shader map was not found in memory, try to load it from the DDC
 		STAT(double OpenColorIOShaderDDCTime = 0);
 		{
diff --git a/Engine/Plugins/Compositing/OpenColorIO/Source/OpenColorIO/Private/OpenColorIOShared.cpp b/Engine/Plugins/Compositing/OpenColorIO/Source/OpenColorIO/Private/OpenColorIOShared.cpp
--- a/Engine/Plugins/Compositing/OpenColorIO/Source/OpenColorIO/Private/OpenColorIOShared.cpp
+++ b/Engine/Plugins/Compositing/OpenColorIO/Source/OpenColorIO/Private/OpenColorIOShared.cpp
@@ -298,11 +298,11 @@ bool FOpenColorIOTransformResource::CacheShaders(const FOpenColorIOShaderMapId&
 				FOpenColorIOShaderMap::LoadFromDerivedDataCache(this, InShaderMapId, InPlatform, GameThreadShaderMap);
 				if (GameThreadShaderMap && GameThreadShaderMap->IsValid())
 				{
-					UE_LOG(LogOpenColorIO, Display, TEXT("Loaded shader %s for OCIO ColorSpace %s from DDC"), *GameThreadShaderMap->GetFriendlyName(), *GetFriendlyName());
+					UE_LOG(LogOpenColorIO, Display, TEXT("Loaded shader %s for transform %s from DDC"), *GameThreadShaderMap->GetFriendlyName(), *GetFriendlyName());
 				}
 				else
 				{
-					UE_LOG(LogOpenColorIO, Display, TEXT("Loading shader for OCIO ColorSpace %s from DDC failed. Shader needs recompile."), *GetFriendlyName());
+					UE_LOG(LogOpenColorIO, Display, TEXT("Shader for transform %s was not loaded from the DDC, needs recompile."), *GetFriendlyName());
 				}
 #endif // WITH_EDITOR
 			}
