OCIO: Minimal set of changes to allow the movie render graph integration to work in non-editor modes.

diff --git a/Engine/Plugins/Compositing/OpenColorIO/Source/OpenColorIO/Private/OpenColorIOConfiguration.cpp b/Engine/Plugins/Compositing/OpenColorIO/Source/OpenColorIO/Private/OpenColorIOConfiguration.cpp
--- a/Engine/Plugins/Compositing/OpenColorIO/Source/OpenColorIO/Private/OpenColorIOConfiguration.cpp
+++ b/Engine/Plugins/Compositing/OpenColorIO/Source/OpenColorIO/Private/OpenColorIOConfiguration.cpp
@@ -13,6 +13,7 @@
 #include "OpenColorIOColorTransform.h"
 #include "OpenColorIOModule.h"
 #include "OpenColorIOSettings.h"
+#include "OpenColorIOWrapper.h"
 #include "TextureResource.h"
 #include "Widgets/Notifications/SNotificationList.h"
 
@@ -22,9 +23,7 @@
 
 #define LOCTEXT_NAMESPACE "OCIOConfiguration"
 
-
 #if WITH_EDITOR
-#include "OpenColorIOWrapper.h"
 #include "DerivedDataCacheInterface.h"
 #include "DirectoryWatcherModule.h"
 #include "IDirectoryWatcher.h"
@@ -349,7 +348,21 @@
 
 FOpenColorIOWrapperConfig* UOpenColorIOConfiguration::GetConfigWrapper() const
 {
-#if WITH_EDITOR
+#if WITH_OCIO
 	return Config.Get();
 #else
 	return nullptr;
+#endif
+}
+
+FOpenColorIOWrapperConfig* UOpenColorIOConfiguration::GetOrCreateConfigWrapper()
+{
+#if WITH_OCIO
+	if (!Config.IsValid())
+	{
+		LoadConfiguration();
+	}
+
+	return Config.Get();
+#else
+	return nullptr;
@@ -589,6 +602,59 @@
 	OpenColorIOConfiguration::SendAnalytics(TEXT("Usage.OpenColorIO.ConfigAssetSaved"), DesiredColorSpaces);
 }
 
+void UOpenColorIOConfiguration::LoadConfiguration()
+{
+	// NOTE: We loosen the editor-only restriction on this method for use
+	// by clients that already link the wrapper library in non-editor modes.
+	// Further release of such restrictions within this plugin is intended,
+	// though all the internal callers of this method remain WITH_EDITOR for now.
+#if WITH_OCIO
+	Config.Reset();
 
+	if (ConfigurationFile.FilePath.IsEmpty())
+	{
+		return;
+	}
+
+	bool bIsBuiltIn = false;
+	FString FilePath = ConfigurationFile.FilePath;
+
+	if (ConfigurationFile.FilePath.StartsWith(TEXT("ocio://")))
+	{
+		bIsBuiltIn = true;
+	}
+	else if (ConfigurationFile.FilePath.Equals(TEXT("$OCIO")) || ConfigurationFile.FilePath.Equals(TEXT("%OCIO%")))
+	{
+		FilePath = FPlatformMisc::GetEnvironmentVariable(TEXT("OCIO"));
+	}
+	else if (ConfigurationFile.FilePath.Contains(TEXT("{Engine}")))
+	{
+		FilePath = FPaths::ConvertRelativePathToFull(ConfigurationFile.FilePath.Replace(TEXT("{Engine}"), *FPaths::EngineDir()));
+	}
+	else if (FPaths::IsRelative(ConfigurationFile.FilePath))
+	{
+		const FString AbsoluteGameDir = FPaths::ConvertRelativePathToFull(FPaths::ProjectDir());
+		FilePath = FPaths::ConvertRelativePathToFull(FPaths::Combine(AbsoluteGameDir, ConfigurationFile.FilePath));
+	}
+
+	Config = MakePimpl<FOpenColorIOWrapperConfig>(FilePath, FOpenColorIOWrapperConfig::WCS_AsInterchangeSpace);
+
+	if (Config->IsValid())
+	{
+		if (!bIsBuiltIn)
+		{
+			StartDirectoryWatch(FilePath);
+		}
+
+		UE_LOG(LogOpenColorIO, Verbose, TEXT("Loaded OCIO configuration file %s"), *FilePath);
+	}
+	else
+	{
+		UE_LOG(LogOpenColorIO, Error, TEXT("Could not load OCIO configuration file %s. Verify that the path is good or that the file is valid."), *FilePath);
+	}
+#endif
+}
+
+
 #if WITH_EDITOR
 
@@ -660,53 +726,6 @@
 	Super::PostEditChangeProperty(PropertyChangedEvent);
 }
 
-void UOpenColorIOConfiguration::LoadConfiguration()
-{
-	Config.Reset();
-
-	if (ConfigurationFile.FilePath.IsEmpty())
-	{
-		return;
-	}
-
-	bool bIsBuiltIn = false;
-	FString FilePath = ConfigurationFile.FilePath;
-
-	if (ConfigurationFile.FilePath.StartsWith(TEXT("ocio://")))
-	{
-		bIsBuiltIn = true;
-	}
-	else if (ConfigurationFile.FilePath.Equals(TEXT("$OCIO")) || ConfigurationFile.FilePath.Equals(TEXT("%OCIO%")))
-	{
-		FilePath = FPlatformMisc::GetEnvironmentVariable(TEXT("OCIO"));
-	}
-	else if(ConfigurationFile.FilePath.Contains(TEXT("{Engine}")))
-	{
-		FilePath = FPaths::ConvertRelativePathToFull(ConfigurationFile.FilePath.Replace(TEXT("{Engine}"), *FPaths::EngineDir()));
-	}
-	else if (FPaths::IsRelative(ConfigurationFile.FilePath))
-	{
-		const FString AbsoluteGameDir = FPaths::ConvertRelativePathToFull(FPaths::ProjectDir());
-		FilePath = FPaths::ConvertRelativePathToFull(FPaths::Combine(AbsoluteGameDir, ConfigurationFile.FilePath));
-	}
-
-	Config = MakePimpl<FOpenColorIOWrapperConfig>(FilePath, FOpenColorIOWrapperConfig::WCS_AsInterchangeSpace);
-
-	if (Config->IsValid())
-	{
-		if (!bIsBuiltIn)
-		{
-			StartDirectoryWatch(FilePath);
-		}
-
-		UE_LOG(LogOpenColorIO, Verbose, TEXT("Loaded OCIO configuration file %s"), *FilePath);
-	}
-	else
-	{
-		UE_LOG(LogOpenColorIO, Error, TEXT("Could not load OCIO configuration file %s. Verify that the path is good or that the file is valid."), *FilePath);
-	}
-}
-
 void UOpenColorIOConfiguration::OnToastCallback(bool bInReloadColorspaces)
 {
 	if (WatchedDirectoryInfo.RawConfigChangedToast.IsValid())
diff --git a/Engine/Plugins/Compositing/OpenColorIO/Source/OpenColorIO/Public/OpenColorIOConfiguration.h b/Engine/Plugins/Compositing/OpenColorIO/Source/OpenColorIO/Public/OpenColorIOConfiguration.h
--- a/Engine/Plugins/Compositing/OpenColorIO/Source/OpenColorIO/Public/OpenColorIOConfiguration.h
+++ b/Engine/Plugins/Compositing/OpenColorIO/Source/OpenColorIO/Public/OpenColorIOConfiguration.h
@@ -66,8 +66,17 @@ public:
 	*/
 	void ConfigPathChangedEvent(const TArray<FFileChangeData>& InFileChanges, const FString InFileMountPath);
 
-	/** Internal only: Replacement for previous `GetNativeConfig_Internal()`, `GetLoadedConfiguration()` and `GetLoadedConfigurationFile()` functions, returning the private implementation of the native OCIO config. */
+	/** 
+	 * Get the private wrapper implementation of the OCIO config.
+	 * Meant for internal use only, caller must have OpenColorIOWrapper module linked.
+	 */
 	FOpenColorIOWrapperConfig* GetConfigWrapper() const;
+
+	/**
+	 * Get or create the private wrapper implementation of the OCIO config.
+	 * Meant for internal use only, caller must have OpenColorIOWrapper module linked.
+	 */
+	FOpenColorIOWrapperConfig* GetOrCreateConfigWrapper();
 	
 	/** Find the color transform object that corresponds to the specified settings, nullptr if not found. */
 	TObjectPtr<const UOpenColorIOColorTransform> FindTransform(const FOpenColorIOColorConversionSettings& InSettings) const;
@@ -105,9 +114,11 @@ public:
 	//~ End UObject interface
 
 private:
-#if WITH_EDITOR
+
+	/** Load the config file to initialize the configuration wrapper object. Automatically called internally. */
 	void LoadConfiguration();
 
+#if WITH_EDITOR
 	/** This method resets the status of Notification dialog and reacts depending on user's choice. */
 	void OnToastCallback(bool bInReloadColorspaces);
 #endif
@@ -144,10 +155,10 @@ private:
 	/** Information about the currently watched directory. Helps us manage the directory change events. */
 	FOCIOConfigWatchedDirInfo WatchedDirectoryInfo;
 
-#if WITH_EDITORONLY_DATA
 	/** Private implementation of the OpenColorIO config object. */
 	TPimplPtr<FOpenColorIOWrapperConfig> Config = nullptr;
 
+#if WITH_EDITORONLY_DATA
 	/** Hash of all of the config content, including relevant external file information. */
 	UPROPERTY()
 	FString ConfigHash;
