OCIO: Initial pass to allow library use in non-editor modes.

diff --git a/Engine/Plugins/Compositing/OpenColorIO/Source/OpenColorIO/OpenColorIO.Build.cs b/Engine/Plugins/Compositing/OpenColorIO/Source/OpenColorIO/OpenColorIO.Build.cs
--- a/Engine/Plugins/Compositing/OpenColorIO/Source/OpenColorIO/OpenColorIO.Build.cs
+++ b/Engine/Plugins/Compositing/OpenColorIO/Source/OpenColorIO/OpenColorIO.Build.cs
@@ -30,8 +30,7 @@ namespace UnrealBuildTool.Rules
 
 			PrivateIncludePathModuleNames.AddRange(
 				new string[] {
-					"DerivedDataCache",
-					"OpenColorIOWrapper" //for OpenColorIOWrapperDefines.h
+					"DerivedDataCache"
 				});
 
 			PublicDependencyModuleNames.AddRange(
@@ -39,7 +38,8 @@ namespace UnrealBuildTool.Rules
 					"Core",
 					"CoreUObject",
 					"Engine",
-					"DeveloperSettings"
+					"DeveloperSettings",
+					"OpenColorIOWrapper"
 				});
 
 			if (Target.bBuildEditor == true)
@@ -49,8 +49,7 @@ namespace UnrealBuildTool.Rules
 					{
 						"TargetPlatform",
 						"EditorFramework",
-						"UnrealEd",
-						"OpenColorIOWrapper"
+						"UnrealEd"
 					});
 			}
 		}
diff --git a/Engine/Plugins/Compositing/OpenColorIO/Source/OpenColorIO/Private/OpenColorIOColorTransform.cpp b/Engine/Plugins/Compositing/OpenColorIO/Source/OpenColorIO/Private/OpenColorIOColorTransform.cpp
--- a/Engine/Plugins/Compositing/OpenColorIO/Source/OpenColorIO/Private/OpenColorIOColorTransform.cpp
+++ b/Engine/Plugins/Compositing/OpenColorIO/Source/OpenColorIO/Private/OpenColorIOColorTransform.cpp
@@ -19,6 +19,8 @@
 #if WITH_EDITOR
 #include "DerivedDataCacheInterface.h"
 #include "Editor.h"
+#endif //WITH_EDITOR
+
 #include "Interfaces/ITargetPlatform.h"
 #include "OpenColorIOWrapper.h"
 #include "OpenColorIODerivedDataVersion.h"
@@ -27,18 +29,21 @@
 namespace
 {
 	// Returns the (native) config wrapper from the configuration object.
-	FOpenColorIOWrapperConfig* GetTransformConfigWrapper(const UOpenColorIOConfiguration* InConfigurationOwner)
+	FOpenColorIOWrapperConfig* GetTransformConfigWrapper(UOpenColorIOConfiguration* InConfigurationOwner)
 	{
 		if (IsValid(InConfigurationOwner))
 		{
+#if WITH_EDITOR
 			return InConfigurationOwner->GetConfigWrapper();
+#else
+			// In non-editor modes, we don't automatically load the config so it needs to be lazily created here.
+			return InConfigurationOwner->GetOrCreateConfigWrapper();
+#endif
 		}
 
 		return nullptr;
 	}
 }
-#endif //WITH_EDITOR
-
 
 void UOpenColorIOColorTransform::SerializeOpenColorIOShaderMaps(const TMap<const ITargetPlatform*, TArray<FOpenColorIOTransformResource*>>* PlatformColorTransformResourcesToSavePtr, FArchive& Ar, TArray<FOpenColorIOTransformResource>&  OutLoadedResources)
 {
@@ -136,20 +141,12 @@ UOpenColorIOColorTransform::UOpenColorIOColorTransform(const FObjectInitializer&
 
 bool UOpenColorIOColorTransform::Initialize(UOpenColorIOConfiguration* InOwner, const FString& InSourceColorSpace, const FString& InDestinationColorSpace, const TMap<FString, FString>& InContextKeyValues)
 {
-#if WITH_EDITOR
 	return Initialize(InSourceColorSpace, InDestinationColorSpace, InContextKeyValues);
-#else
-	return false;
-#endif
 }
 
 bool UOpenColorIOColorTransform::Initialize(UOpenColorIOConfiguration* InOwner, const FString& InSourceColorSpace, const FString& InDisplay, const FString& InView, EOpenColorIOViewTransformDirection InDirection, const TMap<FString, FString>& InContextKeyValues)
 {
-#if WITH_EDITOR
 	return Initialize(InSourceColorSpace, InDisplay, InView, InDirection, InContextKeyValues);
-#else
-	return false;
-#endif
 }
 
 void UOpenColorIOColorTransform::Serialize(FArchive& Ar)
@@ -178,9 +175,6 @@ void UOpenColorIOColorTransform::Serialize(FArchive& Ar)
 	}
 }
 
-
-#if WITH_EDITOR
-
 bool UOpenColorIOColorTransform::Initialize(const FString& InSourceColorSpace, const FString& InDestinationColorSpace, const TMap<FString, FString>& InContextKeyValues)
 {
 
@@ -194,7 +188,9 @@ bool UOpenColorIOColorTransform::Initialize(const FString& InSourceColorSpace, c
 	bIsDisplayViewType = false;
 	ContextKeyValues = InContextKeyValues;
 
-	ProcessTransform();
+#if WITH_EDITOR
+	ProcessTransformForGPU();
+#endif
 	CacheResourceShadersForRendering(true);
 
 	return true;
@@ -215,12 +211,15 @@ bool UOpenColorIOColorTransform::Initialize(const FString& InSourceColorSpace, c
 	DisplayViewDirection = InDirection;
 	ContextKeyValues = InContextKeyValues;
 
-	ProcessTransform();
+#if WITH_EDITOR
+	ProcessTransformForGPU();
+#endif
 	CacheResourceShadersForRendering(true);
 
 	return true;
 }
 
+#if WITH_EDITOR
 void UOpenColorIOColorTransform::CacheResourceShadersForCooking(EShaderPlatform InShaderPlatform, const ITargetPlatform* TargetPlatform, const FString& InShaderHash, const FString& InShaderCode, const FString& InRawConfigHash, TArray<FOpenColorIOTransformResource*>& OutCachedResources)
 {
 	const ERHIFeatureLevel::Type TargetFeatureLevel = GetMaxSupportedFeatureLevel(InShaderPlatform);
@@ -237,9 +236,9 @@ void UOpenColorIOColorTransform::CacheResourceShadersForCooking(EShaderPlatform
 	OutCachedResources.Add(NewResource);
 }
 
-
-void UOpenColorIOColorTransform::ProcessTransform()
+void UOpenColorIOColorTransform::ProcessTransformForGPU()
 {
+#if WITH_OCIO
 	FOpenColorIOWrapperProcessor TransformProcessor;
 	if (GetTransformProcessor(TransformProcessor))
 	{
@@ -290,6 +289,7 @@ void UOpenColorIOColorTransform::ProcessTransform()
 		// Generate shader code
 		GPUProcessor.GetShader(GeneratedShaderHash, GeneratedShader);
 	}
+#endif //WITH_OCIO
 }
 
 TObjectPtr<UTexture> UOpenColorIOColorTransform::CreateTexture3DLUT(const FString& InProcessorIdentifier, const FName& InName, uint32 InLutLength, TextureFilter InFilter, const float* InSourceData)
@@ -408,7 +408,7 @@ TObjectPtr<UTexture> UOpenColorIOColorTransform::CreateTexture1DLUT(const FStrin
 
 	return OutTexture;
 }
-#endif // WITH_EDITOR
+#endif //WITH_EDITOR
 
 void UOpenColorIOColorTransform::CacheResourceShadersForRendering(bool bRegenerateId)
 {
@@ -565,7 +565,6 @@ bool UOpenColorIOColorTransform::IsTransform(const FString& InSourceColorSpace,
 	return false;
 }
 
-#if WITH_EDITOR
 bool UOpenColorIOColorTransform::GetTransformProcessor(FOpenColorIOWrapperProcessor& OutProcessor) const
 {
 	return GetTransformProcessor(GetContextKeyValues(), OutProcessor);
@@ -573,7 +572,8 @@ bool UOpenColorIOColorTransform::GetTransformProcessor(FOpenColorIOWrapperProces
 
 bool UOpenColorIOColorTransform::GetTransformProcessor(const TMap<FString, FString>& InLocalContext, FOpenColorIOWrapperProcessor& OutProcessor) const
 {
-	const UOpenColorIOConfiguration* ConfigurationOwner = Cast<UOpenColorIOConfiguration>(GetOuter());
+#if WITH_OCIO
+	UOpenColorIOConfiguration* ConfigurationOwner = Cast<UOpenColorIOConfiguration>(GetOuter());
 	const FOpenColorIOWrapperConfig* ConfigWrapper = GetTransformConfigWrapper(ConfigurationOwner);
 
 	if (!ConfigWrapper)
@@ -594,41 +594,49 @@ bool UOpenColorIOColorTransform::GetTransformProcessor(const TMap<FString, FStri
 	}
 
 	return OutProcessor.IsValid();
+#else
+	return false;
+#endif
 }
 
 bool UOpenColorIOColorTransform::TransformColor(FLinearColor& InOutColor) const
 {
+#if WITH_OCIO
 	FOpenColorIOWrapperProcessor TransformProcessor;
 	if (GetTransformProcessor(TransformProcessor))
 	{
 		return TransformProcessor.TransformColor(InOutColor);
 	}
+#endif
 
 	return false;
 }
 
 bool UOpenColorIOColorTransform::TransformImage(const FImageView& InOutImage) const
 {
+#if WITH_OCIO
 	FOpenColorIOWrapperProcessor TransformProcessor;
 	if (GetTransformProcessor(TransformProcessor))
 	{
 		return TransformProcessor.TransformImage(InOutImage);
 	}
+#endif
 
 	return false;
 }
 
 bool UOpenColorIOColorTransform::TransformImage(const FImageView& SrcImage, const FImageView& DestImage) const
 {
+#if WITH_OCIO
 	FOpenColorIOWrapperProcessor TransformProcessor;
 	if (GetTransformProcessor(TransformProcessor))
 	{
 		return TransformProcessor.TransformImage(SrcImage, DestImage);
 	}
+#endif
 
 	return false;
 }
-#endif // WITH_EDITOR
 
 bool UOpenColorIOColorTransform::GetDisplayViewDirection(EOpenColorIOViewTransformDirection& OutDirection) const
 {
@@ -727,7 +735,7 @@ void UOpenColorIOColorTransform::PostLoad()
 			// Recently saved transforms will have a valid generated shader which we can use as an indication to skip processing on load.
 			if (GeneratedShader.IsEmpty())
 			{
-				ProcessTransform();
+				ProcessTransformForGPU();
 			}
 #endif
 
@@ -783,7 +791,7 @@ void UOpenColorIOColorTransform::FinishDestroy()
 bool UOpenColorIOColorTransform::UpdateShaderInfo(FString& OutShaderCodeHash, FString& OutShaderCode, FString& OutRawConfigHash)
 {
 	const UOpenColorIOSettings* Settings = GetDefault<UOpenColorIOSettings>();
-	const UOpenColorIOConfiguration* ConfigurationOwner = Cast<UOpenColorIOConfiguration>(GetOuter());
+	UOpenColorIOConfiguration* ConfigurationOwner = Cast<UOpenColorIOConfiguration>(GetOuter());
 	const FOpenColorIOWrapperConfig* ConfigWrapper = GetTransformConfigWrapper(ConfigurationOwner);
 
 	if (ConfigWrapper != nullptr)
diff --git a/Engine/Plugins/Compositing/OpenColorIO/Source/OpenColorIO/Private/OpenColorIOConfiguration.cpp b/Engine/Plugins/Compositing/OpenColorIO/Source/OpenColorIO/Private/OpenColorIOConfiguration.cpp
--- a/Engine/Plugins/Compositing/OpenColorIO/Source/OpenColorIO/Private/OpenColorIOConfiguration.cpp
+++ b/Engine/Plugins/Compositing/OpenColorIO/Source/OpenColorIO/Private/OpenColorIOConfiguration.cpp
@@ -13,6 +13,7 @@
 #include "OpenColorIOColorTransform.h"
 #include "OpenColorIOModule.h"
 #include "OpenColorIOSettings.h"
+#include "OpenColorIOWrapper.h"
 #include "TextureResource.h"
 #include "Widgets/Notifications/SNotificationList.h"
 
@@ -24,7 +25,6 @@
 
 
 #if WITH_EDITOR
-#include "OpenColorIOWrapper.h"
 #include "DerivedDataCacheInterface.h"
 #include "DirectoryWatcherModule.h"
 #include "IDirectoryWatcher.h"
@@ -123,7 +123,7 @@
 
 bool UOpenColorIOConfiguration::Validate() const
 {
-#if WITH_EDITOR
+#if WITH_OCIO
 	if (!ConfigurationFile.FilePath.IsEmpty() && Config.IsValid())
 	{
 		//When loading the configuration file, if any errors are detected, it will throw an exception. Thus, our pointer won't be valid.
@@ -133,10 +133,9 @@
 	return false;
 #else
 	return true;
-#endif // WITH_EDITOR
+#endif // WITH_OCIO
 }
 
-#if WITH_EDITOR
 bool UOpenColorIOConfiguration::TransformColor(const FOpenColorIOColorConversionSettings& InSettings, FLinearColor& InOutColor) const
 {
 	TObjectPtr<const UOpenColorIOColorTransform> Transform = FindTransform(InSettings);
@@ -169,13 +168,12 @@
 
 	return false;
 }
-#endif
 
 void UOpenColorIOConfiguration::ReloadExistingColorspaces(bool bForce)
 {
-#if WITH_EDITOR
 	LoadConfiguration();
 
+#if WITH_OCIO
 	if (Config && Config->IsValid())
 	{
 		FString LoadedConfigHash = Config->GetCacheID();
@@ -349,7 +347,20 @@
 
 FOpenColorIOWrapperConfig* UOpenColorIOConfiguration::GetConfigWrapper() const
 {
-#if WITH_EDITOR
+#if WITH_OCIO
 	return Config.Get();
 #else
 	return nullptr;
+#endif
+}
+
+FOpenColorIOWrapperConfig* UOpenColorIOConfiguration::GetOrCreateConfigWrapper()
+{
+#if WITH_OCIO
+	if (!Config.IsValid())
+	{
+		LoadConfiguration();
+	}
+	return Config.Get();
+#else
+	return nullptr;
@@ -413,7 +424,6 @@
 	return (TransformPtr != nullptr) ? *TransformPtr : nullptr;
 }
 
-#if WITH_EDITOR
 void UOpenColorIOConfiguration::CreateColorTransform(const FString& InSourceColorSpace, const FString& InDestinationColorSpace)
 {
 	if (InSourceColorSpace.IsEmpty() || InDestinationColorSpace.IsEmpty())
@@ -465,7 +475,6 @@
 		UE_LOG(LogOpenColorIO, Warning, TEXT("Could not create color space transform from %s to %s - %s. Verify your OCIO config file, it may have errors in it."), *InSourceColorSpace, *InDisplay, *InView);
 	}
 }
-#endif
 
 void UOpenColorIOConfiguration::CleanupTransforms()
 {
@@ -529,7 +538,12 @@
 void UOpenColorIOConfiguration::PostLoad()
 {
 	Super::PostLoad();
+
+#if WITH_EDITOR
+	// In non-editor modes, the shader & texture resources are already baked.
+	// For CPU transforms, the config is lazily-loaded.
 	ReloadExistingColorspaces();
+#endif
 
 	for (UOpenColorIOColorTransform* Transform : ColorTransforms)
 	{
@@ -589,6 +603,54 @@
 	OpenColorIOConfiguration::SendAnalytics(TEXT("Usage.OpenColorIO.ConfigAssetSaved"), DesiredColorSpaces);
 }
 
+void UOpenColorIOConfiguration::LoadConfiguration()
+{
+#if WITH_OCIO
+	Config.Reset();
 
+	if (ConfigurationFile.FilePath.IsEmpty())
+	{
+		return;
+	}
+
+	bool bIsBuiltIn = false;
+	FString FilePath = ConfigurationFile.FilePath;
+
+	if (ConfigurationFile.FilePath.StartsWith(TEXT("ocio://")))
+	{
+		bIsBuiltIn = true;
+	}
+	else if (ConfigurationFile.FilePath.Equals(TEXT("$OCIO")) || ConfigurationFile.FilePath.Equals(TEXT("%OCIO%")))
+	{
+		FilePath = FPlatformMisc::GetEnvironmentVariable(TEXT("OCIO"));
+	}
+	else if (ConfigurationFile.FilePath.Contains(TEXT("{Engine}")))
+	{
+		FilePath = FPaths::ConvertRelativePathToFull(ConfigurationFile.FilePath.Replace(TEXT("{Engine}"), *FPaths::EngineDir()));
+	}
+	else if (FPaths::IsRelative(ConfigurationFile.FilePath))
+	{
+		const FString AbsoluteGameDir = FPaths::ConvertRelativePathToFull(FPaths::ProjectDir());
+		FilePath = FPaths::ConvertRelativePathToFull(FPaths::Combine(AbsoluteGameDir, ConfigurationFile.FilePath));
+	}
+
+	Config = MakePimpl<FOpenColorIOWrapperConfig>(FilePath, FOpenColorIOWrapperConfig::WCS_AsInterchangeSpace);
+
+	if (Config->IsValid())
+	{
+		if (!bIsBuiltIn)
+		{
+			StartDirectoryWatch(FilePath);
+		}
+
+		UE_LOG(LogOpenColorIO, Verbose, TEXT("Loaded OCIO configuration file %s"), *FilePath);
+	}
+	else
+	{
+		UE_LOG(LogOpenColorIO, Error, TEXT("Could not load OCIO configuration file %s. Verify that the path is good or that the file is valid."), *FilePath);
+	}
+#endif //WITH_OCIO
+}
+
 #if WITH_EDITOR
 
@@ -660,53 +722,6 @@
 	Super::PostEditChangeProperty(PropertyChangedEvent);
 }
 
-void UOpenColorIOConfiguration::LoadConfiguration()
-{
-	Config.Reset();
-
-	if (ConfigurationFile.FilePath.IsEmpty())
-	{
-		return;
-	}
-
-	bool bIsBuiltIn = false;
-	FString FilePath = ConfigurationFile.FilePath;
-
-	if (ConfigurationFile.FilePath.StartsWith(TEXT("ocio://")))
-	{
-		bIsBuiltIn = true;
-	}
-	else if (ConfigurationFile.FilePath.Equals(TEXT("$OCIO")) || ConfigurationFile.FilePath.Equals(TEXT("%OCIO%")))
-	{
-		FilePath = FPlatformMisc::GetEnvironmentVariable(TEXT("OCIO"));
-	}
-	else if(ConfigurationFile.FilePath.Contains(TEXT("{Engine}")))
-	{
-		FilePath = FPaths::ConvertRelativePathToFull(ConfigurationFile.FilePath.Replace(TEXT("{Engine}"), *FPaths::EngineDir()));
-	}
-	else if (FPaths::IsRelative(ConfigurationFile.FilePath))
-	{
-		const FString AbsoluteGameDir = FPaths::ConvertRelativePathToFull(FPaths::ProjectDir());
-		FilePath = FPaths::ConvertRelativePathToFull(FPaths::Combine(AbsoluteGameDir, ConfigurationFile.FilePath));
-	}
-
-	Config = MakePimpl<FOpenColorIOWrapperConfig>(FilePath, FOpenColorIOWrapperConfig::WCS_AsInterchangeSpace);
-
-	if (Config->IsValid())
-	{
-		if (!bIsBuiltIn)
-		{
-			StartDirectoryWatch(FilePath);
-		}
-
-		UE_LOG(LogOpenColorIO, Verbose, TEXT("Loaded OCIO configuration file %s"), *FilePath);
-	}
-	else
-	{
-		UE_LOG(LogOpenColorIO, Error, TEXT("Could not load OCIO configuration file %s. Verify that the path is good or that the file is valid."), *FilePath);
-	}
-}
-
 void UOpenColorIOConfiguration::OnToastCallback(bool bInReloadColorspaces)
 {
 	if (WatchedDirectoryInfo.RawConfigChangedToast.IsValid())
diff --git a/Engine/Plugins/Compositing/OpenColorIO/Source/OpenColorIO/Public/OpenColorIOColorTransform.h b/Engine/Plugins/Compositing/OpenColorIO/Source/OpenColorIO/Public/OpenColorIOColorTransform.h
--- a/Engine/Plugins/Compositing/OpenColorIO/Source/OpenColorIO/Public/OpenColorIOColorTransform.h
+++ b/Engine/Plugins/Compositing/OpenColorIO/Source/OpenColorIO/Public/OpenColorIOColorTransform.h
@@ -39,14 +39,12 @@ public:
 	UE_DEPRECATED(5.3, "This method is deprecated, please use Initialize without the owner argument.")
 	bool Initialize(UOpenColorIOConfiguration* InOwner, const FString& InSourceColorSpace, const FString& InDisplay, const FString& InView, EOpenColorIOViewTransformDirection InDirection, const TMap<FString, FString>& InContextKeyValues = {});
 
-#if WITH_EDITOR
 	/**
 	 * Initialize resources for color space transform. */
 	bool Initialize(const FString& InSourceColorSpace, const FString& InDestinationColorSpace, const TMap<FString, FString>& InContextKeyValues = {});
 	
 	/** Initialize resources for display-view transform. */
 	bool Initialize(const FString& InSourceColorSpace, const FString& InDisplay, const FString& InView, EOpenColorIOViewTransformDirection InDirection, const TMap<FString, FString>& InContextKeyValues = {});
-#endif
 
 	/**
 	 * Serialize LUT data. This will effectively serialize the LUT only when cooking
@@ -87,7 +85,6 @@ public:
 	 */
 	bool IsTransform(const FString& InSourceColorSpace, const FString& InDisplay, const FString& InView, EOpenColorIOViewTransformDirection InDirection) const;
 
-#if WITH_EDITOR
 	/**
 	 * Get the transform processor.
 	 *
@@ -113,7 +110,6 @@ public:
 	
 	/** Apply the color transform from the source image to the destination image. (The destination FImageView is const but what it points at is not.) */
 	bool TransformImage(const FImageView& SrcImage, const FImageView& DestImage) const;
-#endif
 
 	/**
 	 * Get the display view direction type, when applicable.
@@ -179,10 +175,10 @@ protected:
 private:
 #if WITH_EDITOR
 	/**
-	 * Create the transform processor(s) and generate its resources.
+	 * Create the transform processor(s) and generate its resources. Editor-only, in game mode GPU resources are already present.
 	 */
-	void ProcessTransform();
-#endif //WITH_EDITOR
+	void ProcessTransformForGPU();
+#endif
 
 	void FlushResourceShaderMaps();
 
diff --git a/Engine/Plugins/Compositing/OpenColorIO/Source/OpenColorIO/Public/OpenColorIOConfiguration.h b/Engine/Plugins/Compositing/OpenColorIO/Source/OpenColorIO/Public/OpenColorIOConfiguration.h
--- a/Engine/Plugins/Compositing/OpenColorIO/Source/OpenColorIO/Public/OpenColorIOConfiguration.h
+++ b/Engine/Plugins/Compositing/OpenColorIO/Source/OpenColorIO/Public/OpenColorIOConfiguration.h
@@ -45,7 +45,6 @@ public:
 	bool HasDesiredDisplayView(const FOpenColorIODisplayView& DisplayView) const;
 	bool Validate() const;
 
-#if WITH_EDITOR
 	/** Apply the color transform in-place to the specified image. */
 	bool TransformColor(const FOpenColorIOColorConversionSettings& InSettings, FLinearColor& InOutColor) const;
 
@@ -54,7 +53,6 @@ public:
 
 	/** Apply the color transform from the source image to the destination image. (The destination FImageView is const but what it points at is not.) */
 	bool TransformImage(const FOpenColorIOColorConversionSettings& InSettings, const FImageView& SrcImage, const FImageView& DestImage) const;
-#endif
 
 	/** This forces to reload colorspaces and corresponding shaders if those are not loaded already. */
 	UFUNCTION(BlueprintCallable, Category = "OpenColorIO")
@@ -66,18 +64,25 @@ public:
 	*/
 	void ConfigPathChangedEvent(const TArray<FFileChangeData>& InFileChanges, const FString InFileMountPath);
 
-	/** Internal only: Replacement for previous `GetNativeConfig_Internal()`, `GetLoadedConfiguration()` and `GetLoadedConfigurationFile()` functions, returning the private implementation of the native OCIO config. */
+	/**
+	 * Get the private wrapper implementation of the OCIO config.
+	 */
 	FOpenColorIOWrapperConfig* GetConfigWrapper() const;
 	
+	/**
+	 * Get or create the private wrapper implementation of the OCIO config.
+	 * Useful for non-editor modes where the config isn't automatically loaded.
+	 */
+	FOpenColorIOWrapperConfig* GetOrCreateConfigWrapper();
+	
 	/** Find the color transform object that corresponds to the specified settings, nullptr if not found. */
 	TObjectPtr<const UOpenColorIOColorTransform> FindTransform(const FOpenColorIOColorConversionSettings& InSettings) const;
 
 protected:
 
-#if WITH_EDITOR
 	void CreateColorTransform(const FString& InSourceColorSpace, const FString& InDestinationColorSpace);
 	void CreateColorTransform(const FString& InSourceColorSpace, const FString& InDisplay, const FString& InView, EOpenColorIOViewTransformDirection InDirection);
-#endif
+
 	void CleanupTransforms();
 
 	/** Same as above except user can specify the path manually. */
@@ -105,9 +110,11 @@ public:
 	//~ End UObject interface
 
 private:
-#if WITH_EDITOR
+
+	/** Load the config file to initialize the configuration wrapper object. Automatically called internally. */
 	void LoadConfiguration();
 
+#if WITH_EDITOR
 	/** This method resets the status of Notification dialog and reacts depending on user's choice. */
 	void OnToastCallback(bool bInReloadColorspaces);
 #endif
@@ -144,14 +151,12 @@ private:
 	/** Information about the currently watched directory. Helps us manage the directory change events. */
 	FOCIOConfigWatchedDirInfo WatchedDirectoryInfo;
 
-#if WITH_EDITORONLY_DATA
 	/** Private implementation of the OpenColorIO config object. */
 	TPimplPtr<FOpenColorIOWrapperConfig> Config = nullptr;
 
 	/** Hash of all of the config content, including relevant external file information. */
 	UPROPERTY()
 	FString ConfigHash;
-#endif
 };
 
 
diff --git a/Engine/Source/Developer/OpenColorIOWrapper/OpenColorIOWrapper.Build.cs b/Engine/Source/Developer/OpenColorIOWrapper/OpenColorIOWrapper.Build.cs
--- a/Engine/Source/Developer/OpenColorIOWrapper/OpenColorIOWrapper.Build.cs
+++ b/Engine/Source/Developer/OpenColorIOWrapper/OpenColorIOWrapper.Build.cs
@@ -18,29 +18,33 @@ namespace UnrealBuildTool.Rules
 				"Engine", // for TextureDefines.h
 			});
 
-			PrivateDependencyModuleNames.AddRange(new string[]
-			{
-				"ColorManagement",
-				"ImageCore",
-				"OpenColorIOLib",
-			});
-
 			PublicDependencyModuleNames.AddRange(new string[]
 			{
 				"Core",
 			});
 
-			bool bIsPlatformSupported = false;
+			bool bIsSupported = false;
 
-			if (Target.IsInPlatformGroup(UnrealPlatformGroup.Windows) ||
-				Target.IsInPlatformGroup(UnrealPlatformGroup.Unix) ||
-				Target.Platform == UnrealTargetPlatform.Mac)
+			// Because the servers run with NullRHI, we currently ignore OCIO color operations on this build type.
+			if (Target.Type != TargetType.Server)
 			{
-				// Mirror WITH_OCIO_LIB coverage, which suffices for editor use cases.
-				bIsPlatformSupported = true;
+				// Mirror OpenColorIOLib platform coverage
+				if (Target.IsInPlatformGroup(UnrealPlatformGroup.Windows) ||
+					Target.IsInPlatformGroup(UnrealPlatformGroup.Unix) ||
+					Target.Platform == UnrealTargetPlatform.Mac)
+				{
+					PrivateDependencyModuleNames.AddRange(new string[]
+					{
+						"ColorManagement",
+						"ImageCore",
+						"OpenColorIOLib",
+					});
+
+					bIsSupported = true;
+				}
 			}
-
-			PublicDefinitions.Add("WITH_OCIO=" + (bIsPlatformSupported ? "1" : "0"));
+			
+			PublicDefinitions.Add("WITH_OCIO=" + (bIsSupported ? "1" : "0"));
 		}
 	}
 }
diff --git a/Engine/Source/Developer/OpenColorIOWrapper/Private/OpenColorIOWrapperModule.cpp b/Engine/Source/Developer/OpenColorIOWrapper/Private/OpenColorIOWrapperModule.cpp
--- a/Engine/Source/Developer/OpenColorIOWrapper/Private/OpenColorIOWrapperModule.cpp
+++ b/Engine/Source/Developer/OpenColorIOWrapper/Private/OpenColorIOWrapperModule.cpp
@@ -20,10 +20,13 @@ public:
 
 	virtual void ShutdownModule() override
 	{
+#if WITH_OCIO
 		EngineBuiltInConfig.Reset();
+#endif //WITH_OCIO
 	}
 	//~ End IModuleInterface interface
 
+#if WITH_OCIO
 	//~ Begin IOpenColorIOWrapperModule interface
 	virtual FOpenColorIOWrapperEngineConfig& GetEngineBuiltInConfig() override
 	{
@@ -44,6 +47,7 @@ private:
 
 	// Global engine config using the built-in studio config, lazily allocated.
 	TUniquePtr<FOpenColorIOWrapperEngineConfig> EngineBuiltInConfig;
+#endif //WITH_OCIO
 };
 
 IMPLEMENT_MODULE(FOpenColorIOWrapperModule, OpenColorIOWrapper);
diff --git a/Engine/Source/Developer/OpenColorIOWrapper/Public/OpenColorIOWrapperModule.h b/Engine/Source/Developer/OpenColorIOWrapper/Public/OpenColorIOWrapperModule.h
--- a/Engine/Source/Developer/OpenColorIOWrapper/Public/OpenColorIOWrapperModule.h
+++ b/Engine/Source/Developer/OpenColorIOWrapper/Public/OpenColorIOWrapperModule.h
@@ -38,10 +38,12 @@ public:
 		return FModuleManager::Get().IsModuleLoaded(ModuleName);
 	}
 
+#if WITH_OCIO
 	/**
 	 * Returns a minimal dynamically-created native config for engine (working color spaces) conversions.
 	 */
 	virtual FOpenColorIOWrapperEngineConfig& GetEngineBuiltInConfig() = 0;
+#endif //WITH_OCIO
 
 	/** Virtual destructor */
 	virtual ~IOpenColorIOWrapperModule() = default;
