OCIO: Prepare to privatize methods for the next release.

#rnx

diff --git a/Engine/Plugins/Compositing/OpenColorIO/Source/OpenColorIO/Public/OpenColorIOColorTransform.h b/Engine/Plugins/Compositing/OpenColorIO/Source/OpenColorIO/Public/OpenColorIOColorTransform.h
--- a/Engine/Plugins/Compositing/OpenColorIO/Source/OpenColorIO/Public/OpenColorIOColorTransform.h
+++ b/Engine/Plugins/Compositing/OpenColorIO/Source/OpenColorIO/Public/OpenColorIOColorTransform.h
@@ -122,16 +122,19 @@ protected:
 	/**
 	 * Helper function to serialize shader maps for the given color transform resources.
 	 */
+	UE_DEPRECATED_FORGAME(5.4, "Do not use. Will be made private in 5.5")
 	static void SerializeOpenColorIOShaderMaps(const TMap<const ITargetPlatform*, TArray<FOpenColorIOTransformResource*>>* PlatformColorTransformResourcesToSavePtr, FArchive& Ar, TArray<FOpenColorIOTransformResource>&  OutLoadedResources);
 	
 	/**
 	 * Helper function to register serialized shader maps for the given color transform resources
 	 */
+	UE_DEPRECATED_FORGAME(5.4, "Do not use. Will be made private in 5.5")
 	static void ProcessSerializedShaderMaps(UOpenColorIOColorTransform* Owner, TArray<FOpenColorIOTransformResource>& LoadedResources, FOpenColorIOTransformResource* (&OutColorTransformResourcesLoaded)[ERHIFeatureLevel::Num]);
 	
 	/**
 	 * Returns a Guid for the LUT based on its unique identifier, name and the OCIO DDC key.
 	 */
+	UE_DEPRECATED_FORGAME(5.4, "Do not use. Will be made private in 5.5")
 	static void GetOpenColorIOLUTKeyGuid(const FString& InProcessorIdentifier, const FName& InName, FGuid& OutLutGuid );
 
 	UE_DEPRECATED(5.3, "This method is deprecated.")
@@ -143,6 +146,7 @@ protected:
 	/**
 	 * Helper function returning the color space transform name based on source and destination color spaces.
 	 */
+	UE_DEPRECATED_FORGAME(5.4, "Do not use. Will be made private in 5.5")
 	FString GetTransformFriendlyName() const;
 
 #if WITH_EDITOR
@@ -150,16 +154,19 @@ protected:
 	 * Fetch shader code and hash from the OCIO library
 	 * @return: true if shader could be generated from the library
 	 */
+	UE_DEPRECATED_FORGAME(5.4, "Do not use. Will be made private in 5.5")
 	bool UpdateShaderInfo(FString& OutShaderCodeHash, FString& OutShaderCode, FString& OutRawConfigHash);
 
 	/**
 	 * Helper function taking raw 3D LUT data coming from the library and initializing a UVolumeTexture with it.
 	 */
+	UE_DEPRECATED_FORGAME(5.4, "Do not use. Will be made private in 5.5")
 	TObjectPtr<UTexture> CreateTexture3DLUT(const FString& InProcessorIdentifier, const FName& InName, uint32 InLutLength, TextureFilter InFilter, const float* InSourceData);
 
 	/**
 	 * Helper function taking raw 1D LUT data coming from the library and initializing a UTexture with it.
 	 */
+	UE_DEPRECATED_FORGAME(5.4, "Do not use. Will be made private in 5.5")
 	TObjectPtr<UTexture> CreateTexture1DLUT(const FString& InProcessorIdentifier, const FName& InName, uint32 InTextureWidth, uint32 InTextureHeight, TextureFilter InFilter, bool bRedChannelOnly, const float* InSourceData );
 #endif //WITH_EDITOR
 
