OpenColorIO:  Since loading OCIO resources is delayed, a condition has been added to check if the resources are ready to be used.
- fixed a bug where preview did not apply OCIO when PostProcess option was disabled.

diff --git a/Engine/Plugins/Compositing/OpenColorIO/Source/OpenColorIO/Private/OpenColorIOColorTransform.cpp b/Engine/Plugins/Compositing/OpenColorIO/Source/OpenColorIO/Private/OpenColorIOColorTransform.cpp
--- a/Engine/Plugins/Compositing/OpenColorIO/Source/OpenColorIO/Private/OpenColorIOColorTransform.cpp
+++ b/Engine/Plugins/Compositing/OpenColorIO/Source/OpenColorIO/Private/OpenColorIOColorTransform.cpp
@@ -242,6 +242,11 @@ void UOpenColorIOColorTransform::ProcessTransformForGPU()
 	FOpenColorIOWrapperProcessor TransformProcessor;
 	if (GetTransformProcessor(TransformProcessor))
 	{
+		// Note for a ticket in jira [UE-202180]
+		// When the Textures array is emptied, this will cause all existing UTexture objects in the game thread to be deleted.
+		// The resources that were returned earlier using GetRenderResources() and GetRenderPassResources() will also be deleted.
+		// Other plugins store OCIO resources in the FOpenColorIORenderPassResources structure, which uses pointers to FTextureResources
+		// that are used without reference counting. Once deleted in the game thread, they become invalid in the rendering thread.
 		Textures.Reset();
 
 		const UOpenColorIOSettings* Settings = GetDefault<UOpenColorIOSettings>();
@@ -491,7 +496,18 @@ bool UOpenColorIOColorTransform::GetRenderResources(ERHIFeatureLevel::Type InFea
 
 		for (const TPair<int32, TObjectPtr<UTexture>>& Pair : Textures)
 		{
-			OutTextureResources.Add(Pair.Key, Pair.Value->GetResource());
+			if (FTextureResource* TextureResource = Pair.Value->GetResource())
+			{
+				OutTextureResources.Add(Pair.Key, TextureResource);
+			}
+			else
+			{
+				// The texture doesn't exist, we can't get all the resources we need
+				OutShaderResource = nullptr;
+				OutTextureResources.Empty();
+
+				return false;
+			}
 		}
 	}
 	
diff --git a/Engine/Plugins/Compositing/OpenColorIO/Source/OpenColorIO/Private/OpenColorIORendering.cpp b/Engine/Plugins/Compositing/OpenColorIO/Source/OpenColorIO/Private/OpenColorIORendering.cpp
--- a/Engine/Plugins/Compositing/OpenColorIO/Source/OpenColorIO/Private/OpenColorIORendering.cpp
+++ b/Engine/Plugins/Compositing/OpenColorIO/Source/OpenColorIO/Private/OpenColorIORendering.cpp
@@ -42,6 +42,7 @@ void FOpenColorIORendering::AddPass_RenderThread(
 	FGlobalShaderMap* ShaderMap = GetGlobalShaderMap(FeatureLevel);
 	TShaderMapRef<FScreenPassVS> VertexShader(ShaderMap);
 
+	bool bRenderPassSuccessed = false;
 	if (InPassResource.ShaderResource != nullptr)
 	{
 		TShaderRef<FOpenColorIOPixelShader> OCIOPixelShader = InPassResource.ShaderResource->GetShader<FOpenColorIOPixelShader>();
@@ -49,15 +50,20 @@ void FOpenColorIORendering::AddPass_RenderThread(
 		FOpenColorIOPixelShaderParameters* Parameters = GraphBuilder.AllocParameters<FOpenColorIOPixelShaderParameters>();
 		Parameters->InputTexture = Input.Texture;
 		Parameters->InputTextureSampler = TStaticSamplerState<>::GetRHI();
-		OpenColorIOBindTextureResources(Parameters, InPassResource.TextureResources);
-		Parameters->Gamma = InGamma;
-		Parameters->TransformAlpha = (uint32)TransformAlpha;
-		Parameters->RenderTargets[0] = Output.GetRenderTargetBinding();
-		
-		FRDGEventName PassName = RDG_EVENT_NAME("OpenColorIOPass %dx%d (%s)", Output.ViewRect.Width(), Output.ViewRect.Height(), InPassResource.TransformName.IsEmpty() ? TEXT("Unspecified Transform") : *InPassResource.TransformName);
-		AddDrawScreenPass(GraphBuilder, MoveTemp(PassName), ViewInfo, OutputViewport, InputViewport, VertexShader, OCIOPixelShader, Parameters);
+		if (OpenColorIOBindTextureResources(Parameters, InPassResource.TextureResources))
+		{
+			Parameters->Gamma = InGamma;
+			Parameters->TransformAlpha = (uint32)TransformAlpha;
+			Parameters->RenderTargets[0] = Output.GetRenderTargetBinding();
+
+			FRDGEventName PassName = RDG_EVENT_NAME("OpenColorIOPass %dx%d (%s)", Output.ViewRect.Width(), Output.ViewRect.Height(), InPassResource.TransformName.IsEmpty() ? TEXT("Unspecified Transform") : *InPassResource.TransformName);
+			AddDrawScreenPass(GraphBuilder, MoveTemp(PassName), ViewInfo, OutputViewport, InputViewport, VertexShader, OCIOPixelShader, Parameters);
+
+			bRenderPassSuccessed = true;
+		}
 	}
-	else
+
+	if(!bRenderPassSuccessed)
 	{
 		// Fallback pass, printing invalid message across the viewport.
 		TShaderMapRef<FOpenColorIOInvalidPixelShader> OCIOInvalidPixelShader(ShaderMap);
diff --git a/Engine/Plugins/Compositing/OpenColorIO/Source/OpenColorIO/Private/OpenColorIOShader.cpp b/Engine/Plugins/Compositing/OpenColorIO/Source/OpenColorIO/Private/OpenColorIOShader.cpp
--- a/Engine/Plugins/Compositing/OpenColorIO/Source/OpenColorIO/Private/OpenColorIOShader.cpp
+++ b/Engine/Plugins/Compositing/OpenColorIO/Source/OpenColorIO/Private/OpenColorIOShader.cpp
@@ -8,7 +8,7 @@
 #include "TextureResource.h"
 #include "SystemTextures.h"
 
-void OpenColorIOBindTextureResources(FOpenColorIOPixelShaderParameters* Parameters, const TSortedMap<int32, FTextureResource*>& InTextureResources)
+bool OpenColorIOBindTextureResources(FOpenColorIOPixelShaderParameters* Parameters, const TSortedMap<int32, FTextureResource*>& InTextureResources)
 {
 	/**
 	* The engine's shader parameter API doesn't provide name/slot-based updates (in the style of legacy OpenGL),
@@ -23,9 +23,17 @@ void OpenColorIOBindTextureResources(FOpenColorIOPixelShaderParameters* Paramete
 	for ( const TPair<int32, FTextureResource*>& Pair : InTextureResources )
 	{
 		const int32 Index = Pair.Key;
-		ensureMsgf(Index >= 0 && Index < (int32)OpenColorIOShader::MaximumTextureSlots, TEXT("Out of range LUT texture slot index."));
+		if (!ensureMsgf(Index >= 0 && Index < (int32)OpenColorIOShader::MaximumTextureSlots, TEXT("Out of range LUT texture slot index.")))
+		{
+			return false;
+		}
 
 		const FTextureResource* Texture = Pair.Value;
+		if (!ensureMsgf(Texture, TEXT("Texture resource should not be null.")))
+		{
+			return false;
+		}
+
 		const bool bIsLUT3D = Texture->GetTexture3DResource() != nullptr;
 
 		const uint32 BaseOffset = bIsLUT3D ? offsetof(FOpenColorIOPixelShaderParameters, Ocio_lut3d_0) : offsetof(FOpenColorIOPixelShaderParameters, Ocio_lut1d_0);
@@ -34,6 +42,8 @@ void OpenColorIOBindTextureResources(FOpenColorIOPixelShaderParameters* Paramete
 		*((FRHITexture**)(BaseAddress + ParamOffset)) = Texture->TextureRHI.GetReference();
 		*((FRHISamplerState**)(BaseAddress + ParamOffset + sizeof(FRHITexture*))) = Texture->SamplerStateRHI.GetReference();
 	}
+
+	return true;
 }
 
 FRHITexture* OpenColorIOGetMiniFontTexture()
diff --git a/Engine/Plugins/Compositing/OpenColorIO/Source/OpenColorIO/Public/OpenColorIOShader.h b/Engine/Plugins/Compositing/OpenColorIO/Source/OpenColorIO/Public/OpenColorIOShader.h
--- a/Engine/Plugins/Compositing/OpenColorIO/Source/OpenColorIO/Public/OpenColorIOShader.h
+++ b/Engine/Plugins/Compositing/OpenColorIO/Source/OpenColorIO/Public/OpenColorIOShader.h
@@ -103,6 +103,6 @@ public:
 	using FParameters = FOpenColorIOInvalidShaderParameters;
 };
 
-OPENCOLORIO_API void OpenColorIOBindTextureResources(FOpenColorIOPixelShaderParameters* Parameters, const TSortedMap<int32, FTextureResource*>& InTextureResources);
+OPENCOLORIO_API bool OpenColorIOBindTextureResources(FOpenColorIOPixelShaderParameters* Parameters, const TSortedMap<int32, FTextureResource*>& InTextureResources);
 
 FRHITexture* OpenColorIOGetMiniFontTexture();
diff --git a/Engine/Plugins/Runtime/nDisplay/Source/DisplayCluster/Private/Render/Viewport/DisplayClusterViewportManager.cpp b/Engine/Plugins/Runtime/nDisplay/Source/DisplayCluster/Private/Render/Viewport/DisplayClusterViewportManager.cpp
--- a/Engine/Plugins/Runtime/nDisplay/Source/DisplayCluster/Private/Render/Viewport/DisplayClusterViewportManager.cpp
+++ b/Engine/Plugins/Runtime/nDisplay/Source/DisplayCluster/Private/Render/Viewport/DisplayClusterViewportManager.cpp
@@ -653,19 +653,6 @@ void FDisplayClusterViewportManager::InitializeNewFrame()
 {
 	check(IsInGameThread());
 
-	// Before render
-	for (const TSharedPtr<FDisplayClusterViewport, ESPMode::ThreadSafe>& Viewport : ImplGetCurrentRenderFrameViewports())
-	{
-		if (Viewport.IsValid())
-		{
-			// Initialize OCIO resources
-			if (Viewport->GetOpenColorIO().IsValid())
-			{
-				Viewport->GetOpenColorIO()->UpdateOpenColorIORenderPassResources();
-			}
-		}
-	}
-
 	// Handle new frame for warp policies
 	FDisplayClusterWarpPolicyManager WarpPolicyManager;
 	WarpPolicyManager.HandleNewFrame(*this);
diff --git a/Engine/Plugins/Runtime/nDisplay/Source/DisplayCluster/Private/Render/Viewport/DisplayClusterViewportProxy.cpp b/Engine/Plugins/Runtime/nDisplay/Source/DisplayCluster/Private/Render/Viewport/DisplayClusterViewportProxy.cpp
--- a/Engine/Plugins/Runtime/nDisplay/Source/DisplayCluster/Private/Render/Viewport/DisplayClusterViewportProxy.cpp
+++ b/Engine/Plugins/Runtime/nDisplay/Source/DisplayCluster/Private/Render/Viewport/DisplayClusterViewportProxy.cpp
@@ -585,9 +585,10 @@ bool FDisplayClusterViewportProxy::ImplGetResources_RenderThread(const EDisplayC
 
 EDisplayClusterViewportOpenColorIOMode FDisplayClusterViewportProxy::GetOpenColorIOMode() const
 {
-	if (OpenColorIO.IsValid())
+	if (OpenColorIO.IsValid() && OpenColorIO->IsValid_RenderThread())
 	{
-		if (EnumHasAnyFlags(RenderSettingsICVFX.RuntimeFlags, EDisplayClusterViewportRuntimeICVFXFlags::UVLightcard | EDisplayClusterViewportRuntimeICVFXFlags::Lightcard | EDisplayClusterViewportRuntimeICVFXFlags::Chromakey))
+		if (EnumHasAnyFlags(RenderSettingsICVFX.RuntimeFlags, EDisplayClusterViewportRuntimeICVFXFlags::UVLightcard | EDisplayClusterViewportRuntimeICVFXFlags::Lightcard | EDisplayClusterViewportRuntimeICVFXFlags::Chromakey)
+			|| ConfigurationProxy->GetRenderFrameSettings().IsPostProcessDisabled())
 		{
 			// Rendering without post-processing, OCIO is applied last, to the RTT texture of the viewport
 			return EDisplayClusterViewportOpenColorIOMode::Resolved;
diff --git a/Engine/Plugins/Runtime/nDisplay/Source/DisplayCluster/Private/Render/Viewport/DisplayClusterViewport_OpenColorIO.cpp b/Engine/Plugins/Runtime/nDisplay/Source/DisplayCluster/Private/Render/Viewport/DisplayClusterViewport_OpenColorIO.cpp
--- a/Engine/Plugins/Runtime/nDisplay/Source/DisplayCluster/Private/Render/Viewport/DisplayClusterViewport_OpenColorIO.cpp
+++ b/Engine/Plugins/Runtime/nDisplay/Source/DisplayCluster/Private/Render/Viewport/DisplayClusterViewport_OpenColorIO.cpp
@@ -26,25 +26,20 @@ FDisplayClusterViewport_OpenColorIO::FDisplayClusterViewport_OpenColorIO(const F
 FDisplayClusterViewport_OpenColorIO::~FDisplayClusterViewport_OpenColorIO()
 { }
 
-void FDisplayClusterViewport_OpenColorIO::SetupSceneView(FSceneViewFamily& InOutViewFamily, FSceneView& InOutView) const
+void FDisplayClusterViewport_OpenColorIO::SetupSceneView(FSceneViewFamily& InOutViewFamily, FSceneView& InOutView)
 {
-	if (bShaderResourceValid)
+	FOpenColorIORenderPassResources PassResources = FOpenColorIORendering::GetRenderPassResources(ConversionSettings, InOutViewFamily.GetFeatureLevel());
+	if (PassResources.IsValid())
 	{
 		FOpenColorIORendering::PrepareView(InOutViewFamily, InOutView);
 	}
-}
-
-void FDisplayClusterViewport_OpenColorIO::UpdateOpenColorIORenderPassResources()
-{
-	FOpenColorIORenderPassResources PassResources = FOpenColorIORendering::GetRenderPassResources(ConversionSettings, GMaxRHIFeatureLevel);
-
-	bShaderResourceValid = PassResources.IsValid();
 
+	// Update data on rendering thread
 	ENQUEUE_RENDER_COMMAND(ProcessColorSpaceTransform)(
-		[This = SharedThis(this), PassResourcesRenderThread = MoveTemp(PassResources)](FRHICommandListImmediate& RHICmdList)
+		[This = SharedThis(this), ResourcesRenderThread = MoveTemp(PassResources)](FRHICommandListImmediate& RHICmdList)
 		{
 			//Caches render thread resource to be used when applying configuration in PostRenderViewFamily_RenderThread
-			This->CachedResourcesRenderThread = PassResourcesRenderThread;
+			This->CachedResourcesRenderThread = ResourcesRenderThread;
 		}
 	);
 }
diff --git a/Engine/Plugins/Runtime/nDisplay/Source/DisplayCluster/Private/Render/Viewport/DisplayClusterViewport_OpenColorIO.h b/Engine/Plugins/Runtime/nDisplay/Source/DisplayCluster/Private/Render/Viewport/DisplayClusterViewport_OpenColorIO.h
--- a/Engine/Plugins/Runtime/nDisplay/Source/DisplayCluster/Private/Render/Viewport/DisplayClusterViewport_OpenColorIO.h
+++ b/Engine/Plugins/Runtime/nDisplay/Source/DisplayCluster/Private/Render/Viewport/DisplayClusterViewport_OpenColorIO.h
@@ -32,9 +32,6 @@ public:
 	virtual ~FDisplayClusterViewport_OpenColorIO();
 
 public:
-	/** Update render thread resources. */
-	void UpdateOpenColorIORenderPassResources();
-
 	/** Setup view for OCIO.
 	 *
 	 * @param InOutViewFamily - [In,Out] ViewFamily.
@@ -42,7 +39,7 @@ public:
 	 *
 	 * @return - none.
 	 */
-	void SetupSceneView(FSceneViewFamily& InOutViewFamily, FSceneView& InOutView) const;
+	void SetupSceneView(FSceneViewFamily& InOutViewFamily, FSceneView& InOutView);
 
 	/** Add OCIO render pass.
 	 *
@@ -75,6 +72,12 @@ public:
 		return ConversionSettings;
 	}
 
+	/** Returns true if OCIO can be used in the rendering thread. */
+	bool IsValid_RenderThread() const
+	{
+		return CachedResourcesRenderThread.IsValid();
+	}
+
 private:
 	/* returns the DisplayGamma of OCIO for the viewport context.*/
 	float GetDisplayGamma(const FDisplayClusterViewport_Context& InViewportContext) const;
@@ -85,7 +88,4 @@ private:
 
 	/** Configuration to apply during post render callback. */
 	FOpenColorIOColorConversionSettings ConversionSettings;
-
-	/** Shader resources state. */
-	bool bShaderResourceValid = true;
 };
