OCIO: Exposing context on the configuration asset.

diff --git a/Engine/Plugins/Compositing/OpenColorIO/Source/OpenColorIO/Private/OpenColorIOColorTransform.cpp b/Engine/Plugins/Compositing/OpenColorIO/Source/OpenColorIO/Private/OpenColorIOColorTransform.cpp
--- a/Engine/Plugins/Compositing/OpenColorIO/Source/OpenColorIO/Private/OpenColorIOColorTransform.cpp
+++ b/Engine/Plugins/Compositing/OpenColorIO/Source/OpenColorIO/Private/OpenColorIOColorTransform.cpp
@@ -139,14 +139,24 @@ UOpenColorIOColorTransform::UOpenColorIOColorTransform(const FObjectInitializer&
 {
 }
 
-bool UOpenColorIOColorTransform::Initialize(UOpenColorIOConfiguration* InOwner, const FString& InSourceColorSpace, const FString& InDestinationColorSpace, const TMap<FString, FString>& InContextKeyValues)
+bool UOpenColorIOColorTransform::Initialize(UOpenColorIOConfiguration* InOwner, const FString& InSourceColorSpace, const FString& InDestinationColorSpace, const TMap<FString, FString>& /*InContextKeyValues*/)
 {
-	return Initialize(InSourceColorSpace, InDestinationColorSpace, InContextKeyValues);
+	return Initialize(InSourceColorSpace, InDestinationColorSpace);
 }
 
-bool UOpenColorIOColorTransform::Initialize(UOpenColorIOConfiguration* InOwner, const FString& InSourceColorSpace, const FString& InDisplay, const FString& InView, EOpenColorIOViewTransformDirection InDirection, const TMap<FString, FString>& InContextKeyValues)
+bool UOpenColorIOColorTransform::Initialize(UOpenColorIOConfiguration* InOwner, const FString& InSourceColorSpace, const FString& InDisplay, const FString& InView, EOpenColorIOViewTransformDirection InDirection, const TMap<FString, FString>& /*InContextKeyValues*/)
 {
-	return Initialize(InSourceColorSpace, InDisplay, InView, InDirection, InContextKeyValues);
+	return Initialize(InSourceColorSpace, InDisplay, InView, InDirection);
+}
+
+bool UOpenColorIOColorTransform::Initialize(const FString& InSourceColorSpace, const FString& InDestinationColorSpace, const TMap<FString, FString>& /*InContextKeyValues*/)
+{
+	return Initialize(InSourceColorSpace, InDestinationColorSpace);
+}
+
+bool UOpenColorIOColorTransform::Initialize(const FString& InSourceColorSpace, const FString& InDisplay, const FString& InView, EOpenColorIOViewTransformDirection InDirection, const TMap<FString, FString>& /*InContextKeyValues*/)
+{
+	return Initialize(InSourceColorSpace, InDisplay, InView, InDirection);
 }
 
 void UOpenColorIOColorTransform::Serialize(FArchive& Ar)
@@ -175,7 +185,7 @@ void UOpenColorIOColorTransform::Serialize(FArchive& Ar)
 	}
 }
 
-bool UOpenColorIOColorTransform::Initialize(const FString& InSourceColorSpace, const FString& InDestinationColorSpace, const TMap<FString, FString>& InContextKeyValues)
+bool UOpenColorIOColorTransform::Initialize(const FString& InSourceColorSpace, const FString& InDestinationColorSpace)
 {
 
 	if (InSourceColorSpace.IsEmpty() || InDestinationColorSpace.IsEmpty())
@@ -186,7 +196,6 @@ bool UOpenColorIOColorTransform::Initialize(const FString& InSourceColorSpace, c
 	SourceColorSpace = InSourceColorSpace;
 	DestinationColorSpace = InDestinationColorSpace;
 	bIsDisplayViewType = false;
-	ContextKeyValues = InContextKeyValues;
 
 #if WITH_EDITOR
 	ProcessTransformForGPU();
@@ -196,7 +205,7 @@ bool UOpenColorIOColorTransform::Initialize(const FString& InSourceColorSpace, c
 	return true;
 }
 
-bool UOpenColorIOColorTransform::Initialize(const FString& InSourceColorSpace, const FString& InDisplay, const FString& InView, EOpenColorIOViewTransformDirection InDirection, const TMap<FString, FString>& InContextKeyValues)
+bool UOpenColorIOColorTransform::Initialize(const FString& InSourceColorSpace, const FString& InDisplay, const FString& InView, EOpenColorIOViewTransformDirection InDirection)
 {
 	if (InSourceColorSpace.IsEmpty() || InDisplay.IsEmpty() || InView.IsEmpty())
 	{
@@ -209,7 +218,6 @@ bool UOpenColorIOColorTransform::Initialize(const FString& InSourceColorSpace, c
 	View = InView;
 	bIsDisplayViewType = true;
 	DisplayViewDirection = InDirection;
-	ContextKeyValues = InContextKeyValues;
 
 #if WITH_EDITOR
 	ProcessTransformForGPU();
@@ -424,7 +432,7 @@ void UOpenColorIOColorTransform::CacheResourceShadersForRendering(bool bRegenera
 
 	if (FApp::CanEverRender())
 	{
-		const UOpenColorIOConfiguration* ConfigurationOwner = Cast<UOpenColorIOConfiguration>(GetOuter());
+		const UOpenColorIOConfiguration* ConfigurationOwner = GetTypedOuter<UOpenColorIOConfiguration>();
 		if (IsValid(ConfigurationOwner))
 		{
 			//OCIO shaders are simple, we should be compatible with any feature levels. Use the levels required for materials.
@@ -584,7 +592,7 @@ bool UOpenColorIOColorTransform::IsTransform(const FString& InSourceColorSpace,
 bool UOpenColorIOColorTransform::GetTransformProcessor(FOpenColorIOWrapperProcessor& OutProcessor) const
 {
 #if WITH_OCIO
-	UOpenColorIOConfiguration* ConfigurationOwner = Cast<UOpenColorIOConfiguration>(GetOuter());
+	UOpenColorIOConfiguration* ConfigurationOwner = GetTypedOuter<UOpenColorIOConfiguration>();
 	const FOpenColorIOWrapperConfig* ConfigWrapper = GetTransformConfigWrapper(ConfigurationOwner);
 
 	if (!ConfigWrapper)
@@ -597,11 +605,11 @@ bool UOpenColorIOColorTransform::GetTransformProcessor(FOpenColorIOWrapperProces
 	EOpenColorIOViewTransformDirection CurrentDisplayViewDirection;
 	if (GetDisplayViewDirection(CurrentDisplayViewDirection))
 	{
-		OutProcessor = FOpenColorIOWrapperProcessor(ConfigWrapper, SourceColorSpace, Display, View, CurrentDisplayViewDirection == EOpenColorIOViewTransformDirection::Inverse, GetContextKeyValues());
+		OutProcessor = FOpenColorIOWrapperProcessor(ConfigWrapper, SourceColorSpace, Display, View, CurrentDisplayViewDirection == EOpenColorIOViewTransformDirection::Inverse, ConfigurationOwner->Context);
 	}
 	else
 	{
-		OutProcessor = FOpenColorIOWrapperProcessor(ConfigWrapper, SourceColorSpace, DestinationColorSpace, GetContextKeyValues());
+		OutProcessor = FOpenColorIOWrapperProcessor(ConfigWrapper, SourceColorSpace, DestinationColorSpace, ConfigurationOwner->Context);
 	}
 
 	return OutProcessor.IsValid();
@@ -669,6 +677,17 @@ void UOpenColorIOColorTransform::AllColorTransformsCacheResourceShadersForRender
 	}
 }
 
+TMap<FString, FString> UOpenColorIOColorTransform::GetContextKeyValues() const
+{
+	UOpenColorIOConfiguration* ConfigurationOwner = GetTypedOuter<UOpenColorIOConfiguration>();
+	if (IsValid(ConfigurationOwner))
+	{
+		return ConfigurationOwner->Context;
+	}
+
+	return TMap<FString, FString>();
+}
+
 FString UOpenColorIOColorTransform::GetTransformFriendlyName() const
 {
 	if (bIsDisplayViewType)
@@ -722,7 +741,7 @@ void UOpenColorIOColorTransform::PostLoad()
 		}
 	}
 
-	UOpenColorIOConfiguration* ConfigurationOwner = Cast<UOpenColorIOConfiguration>(GetOuter());
+	UOpenColorIOConfiguration* ConfigurationOwner = GetTypedOuter<UOpenColorIOConfiguration>();
 
 	//To be able to fetch OCIO data, make sure our config owner has been postloaded.
 	if (ConfigurationOwner)
@@ -802,7 +821,7 @@ void UOpenColorIOColorTransform::FinishDestroy()
 bool UOpenColorIOColorTransform::UpdateShaderInfo(FString& OutShaderCodeHash, FString& OutShaderCode, FString& OutRawConfigHash)
 {
 	const UOpenColorIOSettings* Settings = GetDefault<UOpenColorIOSettings>();
-	UOpenColorIOConfiguration* ConfigurationOwner = Cast<UOpenColorIOConfiguration>(GetOuter());
+	UOpenColorIOConfiguration* ConfigurationOwner = GetTypedOuter<UOpenColorIOConfiguration>();
 	const FOpenColorIOWrapperConfig* ConfigWrapper = GetTransformConfigWrapper(ConfigurationOwner);
 
 	if (ConfigWrapper != nullptr)
@@ -846,7 +865,7 @@ void UOpenColorIOColorTransform::BeginCacheForCookedPlatformData(const ITargetPl
 			//Need to re-update shader data when cooking, they may not have been previously fetched.
 			bValidShaderInfo = UpdateShaderInfo(ShaderCodeHash, ShaderCode, RawConfigHash);
 		}
-		else if(const UOpenColorIOConfiguration* ConfigurationOwner = Cast<UOpenColorIOConfiguration>(GetOuter()))
+		else if(const UOpenColorIOConfiguration* ConfigurationOwner = GetTypedOuter<UOpenColorIOConfiguration>())
 		{
 			ShaderCodeHash = GeneratedShaderHash;
 			ShaderCode = GeneratedShader;
diff --git a/Engine/Plugins/Compositing/OpenColorIO/Source/OpenColorIO/Private/OpenColorIOConfiguration.cpp b/Engine/Plugins/Compositing/OpenColorIO/Source/OpenColorIO/Private/OpenColorIOConfiguration.cpp
--- a/Engine/Plugins/Compositing/OpenColorIO/Source/OpenColorIO/Private/OpenColorIOConfiguration.cpp
+++ b/Engine/Plugins/Compositing/OpenColorIO/Source/OpenColorIO/Private/OpenColorIOConfiguration.cpp
@@ -559,9 +559,9 @@ void UOpenColorIOConfiguration::GetAssetRegistryTags(TArray<FAssetRegistryTag>&
 	PRAGMA_ENABLE_DEPRECATION_WARNINGS;
 }
 
-void UOpenColorIOConfiguration::GetAssetRegistryTags(FAssetRegistryTagsContext Context) const
+void UOpenColorIOConfiguration::GetAssetRegistryTags(FAssetRegistryTagsContext RegistryTagsContext) const
 {
-	Super::GetAssetRegistryTags(Context);
+	Super::GetAssetRegistryTags(RegistryTagsContext);
 
 	FString Description;
 	if (ConfigurationFile.FilePath.IsEmpty())
@@ -577,7 +577,7 @@ void UOpenColorIOConfiguration::GetAssetRegistryTags(FAssetRegistryTagsContext C
 		Description = TEXT("Configuration: ") + ConfigurationFile.FilePath;
 	}
 
-	Context.AddTag(FAssetRegistryTag(TEXT("ConfigurationFile"), Description, FAssetRegistryTag::TT_Hidden));
+	RegistryTagsContext.AddTag(FAssetRegistryTag(TEXT("ConfigurationFile"), Description, FAssetRegistryTag::TT_Hidden));
 }
 
 #if WITH_EDITORONLY_DATA
diff --git a/Engine/Plugins/Compositing/OpenColorIO/Source/OpenColorIO/Public/OpenColorIOColorTransform.h b/Engine/Plugins/Compositing/OpenColorIO/Source/OpenColorIO/Public/OpenColorIOColorTransform.h
--- a/Engine/Plugins/Compositing/OpenColorIO/Source/OpenColorIO/Public/OpenColorIOColorTransform.h
+++ b/Engine/Plugins/Compositing/OpenColorIO/Source/OpenColorIO/Public/OpenColorIOColorTransform.h
@@ -38,13 +38,16 @@ public:
 	bool Initialize(UOpenColorIOConfiguration* InOwner, const FString& InSourceColorSpace, const FString& InDestinationColorSpace, const TMap<FString, FString>& InContextKeyValues = {});
 	UE_DEPRECATED(5.3, "This method is deprecated, please use Initialize without the owner argument.")
 	bool Initialize(UOpenColorIOConfiguration* InOwner, const FString& InSourceColorSpace, const FString& InDisplay, const FString& InView, EOpenColorIOViewTransformDirection InDirection, const TMap<FString, FString>& InContextKeyValues = {});
+	UE_DEPRECATED(5.4, "This method is deprecated, please use Initialize without the context argument.")
+	bool Initialize(const FString& InSourceColorSpace, const FString& InDestinationColorSpace, const TMap<FString, FString>& InContextKeyValues);
+	UE_DEPRECATED(5.4, "This method is deprecated, please use Initialize without the context argument.")
+	bool Initialize(const FString& InSourceColorSpace, const FString& InDisplay, const FString& InView, EOpenColorIOViewTransformDirection InDirection, const TMap<FString, FString>& InContextKeyValues);
 
-	/**
-	 * Initialize resources for color space transform. */
-	bool Initialize(const FString& InSourceColorSpace, const FString& InDestinationColorSpace, const TMap<FString, FString>& InContextKeyValues = {});
+	/** Initialize resources for color space transform. */
+	bool Initialize(const FString& InSourceColorSpace, const FString& InDestinationColorSpace);
 	
 	/** Initialize resources for display-view transform. */
-	bool Initialize(const FString& InSourceColorSpace, const FString& InDisplay, const FString& InView, EOpenColorIOViewTransformDirection InDirection, const TMap<FString, FString>& InContextKeyValues = {});
+	bool Initialize(const FString& InSourceColorSpace, const FString& InDisplay, const FString& InView, EOpenColorIOViewTransformDirection InDirection);
 
 	/**
 	 * Serialize LUT data. This will effectively serialize the LUT only when cooking
@@ -113,9 +116,8 @@ public:
 	// For all ColorTransforms, UOpenColorIOColorTransform::CacheResourceShadersForRendering
 	static void AllColorTransformsCacheResourceShadersForRendering();
 
-	// Get the context key-values.
-	const TMap<FString, FString>& GetContextKeyValues() const { return ContextKeyValues; }
-
+	// Get the owner config's context key-values.
+	TMap<FString, FString> GetContextKeyValues() const;
 
 protected:
 
@@ -263,9 +265,6 @@ private:
 	
 	FOpenColorIOTransformResource* ColorTransformResources[ERHIFeatureLevel::Num];
 
-	/** Key-value string pairs used to define the processor's context. */
-	TMap<FString, FString> ContextKeyValues;
-
 	FRenderCommandFence ReleaseFence;
 
 #if WITH_EDITORONLY_DATA
diff --git a/Engine/Plugins/Compositing/OpenColorIO/Source/OpenColorIO/Public/OpenColorIOConfiguration.h b/Engine/Plugins/Compositing/OpenColorIO/Source/OpenColorIO/Public/OpenColorIOConfiguration.h
--- a/Engine/Plugins/Compositing/OpenColorIO/Source/OpenColorIO/Public/OpenColorIOConfiguration.h
+++ b/Engine/Plugins/Compositing/OpenColorIO/Source/OpenColorIO/Public/OpenColorIOConfiguration.h
@@ -96,7 +96,7 @@ public:
 	//~ Begin UObject interface
 	virtual void PostInitProperties() override;
 	virtual void PostLoad() override;
-	virtual void GetAssetRegistryTags(FAssetRegistryTagsContext Context) const override;
+	virtual void GetAssetRegistryTags(FAssetRegistryTagsContext RegistryTagsContext) const override;
 	UE_DEPRECATED(5.4, "Implement the version that takes FAssetRegistryTagsContext instead.")
 	virtual void GetAssetRegistryTags(TArray<FAssetRegistryTag>& OutTags) const override;
 #if WITH_EDITORONLY_DATA
@@ -126,12 +126,18 @@ public:
 	UPROPERTY(BlueprintReadWrite, EditAnywhere, Category = "Config", meta = (FilePathFilter = "Config Files (*.ocio, *.ocioz)|*.ocio;*.ocioz", RelativeToGameDir))
 	FFilePath ConfigurationFile;
 
-	UPROPERTY(BlueprintReadWrite, EditAnywhere, Category = "ColorSpace")
+	UPROPERTY(BlueprintReadWrite, EditAnywhere, Category = "Transform")
 	TArray<FOpenColorIOColorSpace> DesiredColorSpaces;
 
-	UPROPERTY(BlueprintReadWrite, EditAnywhere, Category = "ColorSpace", DisplayName="Desired Display-Views")
+	UPROPERTY(BlueprintReadWrite, EditAnywhere, Category = "Transform", DisplayName="Desired Display-Views")
 	TArray<FOpenColorIODisplayView> DesiredDisplayViews;
 
+	/**
+	* OCIO context of key-value string pairs, typically used to apply shot-specific looks (such as a CDL color correction, or a 1D grade LUT).
+	*/
+	UPROPERTY(EditAnywhere, BlueprintReadWrite, Category = "Transform")
+	TMap<FString, FString> Context;
+
 private:
 
 	UPROPERTY()
diff --git a/Engine/Plugins/Compositing/OpenColorIO/Source/OpenColorIOEditor/Private/OpenColorIOColorSpaceConversionCustomization.cpp b/Engine/Plugins/Compositing/OpenColorIO/Source/OpenColorIOEditor/Private/OpenColorIOColorSpaceConversionCustomization.cpp
--- a/Engine/Plugins/Compositing/OpenColorIO/Source/OpenColorIOEditor/Private/OpenColorIOColorSpaceConversionCustomization.cpp
+++ b/Engine/Plugins/Compositing/OpenColorIO/Source/OpenColorIOEditor/Private/OpenColorIOColorSpaceConversionCustomization.cpp
@@ -79,7 +79,9 @@ void FOpenColorIOColorConversionSettingsCustomization::CustomizeChildren(TShared
 			{
 				StructBuilder.AddProperty(ChildHandle).IsEnabled(true).ShowPropertyButtons(false);
 
-				ChildHandle->SetOnPropertyValueChanged(FSimpleDelegate::CreateLambda([this]
+				const TSharedRef<IPropertyUtilities> PropertyUtilities = StructCustomizationUtils.GetPropertyUtilities().ToSharedRef();
+
+				ChildHandle->SetOnPropertyValueChanged(FSimpleDelegate::CreateLambda([this, PropertyUtilities]
 					{
 						if (FOpenColorIOColorConversionSettings* ColorSpaceConversion = GetConversionSettings())
 						{
@@ -88,6 +90,8 @@ void FOpenColorIOColorConversionSettingsCustomization::CustomizeChildren(TShared
 							TransformPicker[OCIO_Src]->SetConfiguration(ColorSpaceConversion->ConfigurationSource);
 							TransformPicker[OCIO_Dst]->SetConfiguration(ColorSpaceConversion->ConfigurationSource);
 						}
+
+						PropertyUtilities->ForceRefresh();
 					}));
 			}
 			else if (ChildHandle->GetProperty()->GetFName() == GET_MEMBER_NAME_CHECKED(FOpenColorIOColorConversionSettings, SourceColorSpace))
@@ -235,6 +239,31 @@ void FOpenColorIOColorConversionSettingsCustomization::CustomizeChildren(TShared
 				})
 			];
 		}
+
+		if (IsValid(ColorSpaceConversion->ConfigurationSource))
+		{
+			const TArray<UObject*> ConfigurationObjects = { ColorSpaceConversion->ConfigurationSource.Get() };
+			
+			IDetailPropertyRow* ContextRow = StructBuilder.AddExternalObjectProperty(ConfigurationObjects, GET_MEMBER_NAME_CHECKED(UOpenColorIOConfiguration, Context));
+			if (ContextRow)
+			{
+				ContextRow->DisplayName(LOCTEXT("BaseContext", "Base Context"));
+				ContextRow->IsEnabled(false);
+				ContextRow->Visibility(MakeAttributeLambda([this]()
+					{
+						if (FOpenColorIOColorConversionSettings* ColorSpaceConversion = GetConversionSettings())
+						{
+							if (IsValid(ColorSpaceConversion->ConfigurationSource) && !ColorSpaceConversion->ConfigurationSource->Context.IsEmpty())
+							{
+								return EVisibility::Visible;
+							}
+						}
+
+						return EVisibility::Hidden;
+					})
+				);
+			}
+		}
 	}
 
 	InStructPropertyHandle->MarkHiddenByCustomization();
diff --git a/Engine/Plugins/MovieScene/MovieRenderPipeline/Source/MovieRenderPipelineRenderPasses/Private/MovieGraphImageSequenceOutputNode.cpp b/Engine/Plugins/MovieScene/MovieRenderPipeline/Source/MovieRenderPipelineRenderPasses/Private/MovieGraphImageSequenceOutputNode.cpp
--- a/Engine/Plugins/MovieScene/MovieRenderPipeline/Source/MovieRenderPipelineRenderPasses/Private/MovieGraphImageSequenceOutputNode.cpp
+++ b/Engine/Plugins/MovieScene/MovieRenderPipeline/Source/MovieRenderPipelineRenderPasses/Private/MovieGraphImageSequenceOutputNode.cpp
@@ -88,7 +88,7 @@ namespace UE::MovieGraph::Private
 	};
 
 	/**
-	 * Convenience function to resolve an OpenColorIO context with supported tokens. Editor-only.
+	 * Convenience function to resolve an OpenColorIO context with supported tokens.
 	 *
 	 * @return The resolved key/value context.
 	*/
@@ -371,8 +371,18 @@ void UMovieGraphImageSequenceOutputNode::OnReceiveImageDataImpl(UMovieGraphPipel
 #if WITH_OCIO
 		if (ParentNode->OCIOConfiguration.bIsEnabled && Payload->bAllowOCIO)
 		{
-			TMap<FString, FString> ResolvedOCIOContext = UE::MovieGraph::Private::ResolveOpenColorIOContext(
-				ParentNode->OCIOContext,
+			TMap<FString, FString> ResolvedOCIOContext;
+
+			const TObjectPtr<UOpenColorIOConfiguration>& ConfigurationAsset = ParentNode->OCIOConfiguration.ColorConfiguration.ConfigurationSource;
+			if (IsValid(ConfigurationAsset))
+			{
+				ResolvedOCIOContext = ConfigurationAsset->Context;
+			}
+
+			ResolvedOCIOContext.Append(ParentNode->OCIOContext);
+
+			ResolvedOCIOContext = UE::MovieGraph::Private::ResolveOpenColorIOContext(
+				ResolvedOCIOContext,
 				RenderData.Key,
 				InPipeline,
 				InRawFrameData->EvaluatedConfig.Get(),
diff --git a/Engine/Plugins/MovieScene/MovieRenderPipeline/Source/MovieRenderPipelineRenderPasses/Public/MovieGraphImageSequenceOutputNode.h b/Engine/Plugins/MovieScene/MovieRenderPipeline/Source/MovieRenderPipelineRenderPasses/Public/MovieGraphImageSequenceOutputNode.h
--- a/Engine/Plugins/MovieScene/MovieRenderPipeline/Source/MovieRenderPipelineRenderPasses/Public/MovieGraphImageSequenceOutputNode.h
+++ b/Engine/Plugins/MovieScene/MovieRenderPipeline/Source/MovieRenderPipelineRenderPasses/Public/MovieGraphImageSequenceOutputNode.h
@@ -51,6 +51,10 @@ public:
 
 	/**
 	* OCIO context of key-value string pairs, typically used to apply shot-specific looks (such as a CDL color correction, or a 1D grade LUT).
+	* 
+	* Notes:
+	* 1) If a configuration asset base context was set, it remains active but can be overridden here with new key-values.
+	* 2) Format tokens such as {shot_name} are supported and will get resolved before submission.
 	*/
 	UPROPERTY(EditAnywhere, BlueprintReadWrite, Category = "OCIO", DisplayName = "OCIO Context", meta = (EditCondition = "bOverride_OCIOContext"))
 	TMap<FString, FString> OCIOContext;
