OCIO: Prevent shader compilation crash on empty shader code due to an invalid transform in configs.

Further privatize unnecessary public methods, and prevent leak in unused code path.

diff --git a/Engine/Plugins/Compositing/OpenColorIO/Source/OpenColorIO/Private/OpenColorIOColorTransform.cpp b/Engine/Plugins/Compositing/OpenColorIO/Source/OpenColorIO/Private/OpenColorIOColorTransform.cpp
--- a/Engine/Plugins/Compositing/OpenColorIO/Source/OpenColorIO/Private/OpenColorIOColorTransform.cpp
+++ b/Engine/Plugins/Compositing/OpenColorIO/Source/OpenColorIO/Private/OpenColorIOColorTransform.cpp
@@ -200,7 +200,8 @@ bool UOpenColorIOColorTransform::Initialize(const FString& InSourceColorSpace, c
 #if WITH_EDITOR
 	ProcessTransformForGPU();
 #endif
-	CacheResourceShadersForRendering(true);
+	FlushResourceShaderMaps();
+	CacheResourceShadersForRendering();
 
 	return true;
 }
@@ -222,7 +223,8 @@ bool UOpenColorIOColorTransform::Initialize(const FString& InSourceColorSpace, c
 #if WITH_EDITOR
 	ProcessTransformForGPU();
 #endif
-	CacheResourceShadersForRendering(true);
+	FlushResourceShaderMaps();
+	CacheResourceShadersForRendering();
 
 	return true;
 }
@@ -423,13 +425,8 @@ TObjectPtr<UTexture> UOpenColorIOColorTransform::CreateTexture1DLUT(const FStrin
 }
 #endif //WITH_EDITOR
 
-void UOpenColorIOColorTransform::CacheResourceShadersForRendering(bool bRegenerateId)
+void UOpenColorIOColorTransform::CacheResourceShadersForRendering(bool /*bRegenerateId*/)
 {
-	if (bRegenerateId)
-	{
-		FlushResourceShaderMaps();
-	}
-
 	if (FApp::CanEverRender())
 	{
 		const UOpenColorIOConfiguration* ConfigurationOwner = GetTypedOuter<UOpenColorIOConfiguration>();
@@ -453,6 +450,11 @@ void UOpenColorIOColorTransform::CacheResourceShadersForRendering(bool bRegenera
 				FString RawConfigHash;
 				FName AssetPath;
 #if WITH_EDITOR
+				if (GeneratedShader.IsEmpty())
+				{
+					continue;
+				}
+
 				ShaderCodeHash = GeneratedShaderHash;
 				ShaderCode = GeneratedShader;
 				RawConfigHash = ConfigurationOwner->GetConfigWrapper()->GetCacheID();
@@ -673,7 +675,7 @@ void UOpenColorIOColorTransform::AllColorTransformsCacheResourceShadersForRender
 	{
 		UOpenColorIOColorTransform* Transform = *It;
 
-		Transform->CacheResourceShadersForRendering(false);
+		Transform->CacheResourceShadersForRendering();
 	}
 }
 
@@ -718,7 +720,6 @@ void UOpenColorIOColorTransform::FlushResourceShaderMaps()
 			if (ColorTransformResources[Index])
 			{
 				ColorTransformResources[Index]->ReleaseShaderMap();
-				ColorTransformResources[Index] = nullptr;
 			}
 		}
 	}
@@ -769,7 +770,7 @@ void UOpenColorIOColorTransform::PostLoad()
 			}
 #endif
 
-			CacheResourceShadersForRendering(false);
+			CacheResourceShadersForRendering();
 		}
 		else
 		{
@@ -855,7 +856,6 @@ void UOpenColorIOColorTransform::BeginCacheForCookedPlatformData(const ITargetPl
 
 	if (DesiredShaderFormats.Num() > 0)
 	{
-		bool bValidShaderInfo = false;
 		FString ShaderCodeHash;
 		FString ShaderCode;
 		FString RawConfigHash;
@@ -863,17 +863,16 @@ void UOpenColorIOColorTransform::BeginCacheForCookedPlatformData(const ITargetPl
 		if (GeneratedShader.IsEmpty())
 		{
 			//Need to re-update shader data when cooking, they may not have been previously fetched.
-			bValidShaderInfo = UpdateShaderInfo(ShaderCodeHash, ShaderCode, RawConfigHash);
+			UpdateShaderInfo(ShaderCodeHash, ShaderCode, RawConfigHash);
 		}
 		else if(const UOpenColorIOConfiguration* ConfigurationOwner = GetTypedOuter<UOpenColorIOConfiguration>())
 		{
 			ShaderCodeHash = GeneratedShaderHash;
 			ShaderCode = GeneratedShader;
 			RawConfigHash = ConfigurationOwner->GetConfigWrapper()->GetCacheID();
-			bValidShaderInfo = true;
 		}
 
-		if (bValidShaderInfo)
+		if (!ShaderCode.IsEmpty())
 		{
 			// Cache for all the shader formats that the cooking target requires
 			for (int32 FormatIndex = 0; FormatIndex < DesiredShaderFormats.Num(); FormatIndex++)
diff --git a/Engine/Plugins/Compositing/OpenColorIO/Source/OpenColorIO/Public/OpenColorIOColorTransform.h b/Engine/Plugins/Compositing/OpenColorIO/Source/OpenColorIO/Public/OpenColorIOColorTransform.h
--- a/Engine/Plugins/Compositing/OpenColorIO/Source/OpenColorIO/Public/OpenColorIOColorTransform.h
+++ b/Engine/Plugins/Compositing/OpenColorIO/Source/OpenColorIO/Public/OpenColorIOColorTransform.h
@@ -62,7 +62,10 @@ public:
 	 * Cache resource shaders for rendering.
 	 * If a matching shader map is not found in memory or the DDC, a new one will be compiled.
 	 */
-	void CacheResourceShadersForRendering(bool bRegenerateId);
+	UE_DEPRECATED_FORGAME(5.4, "Do not use. Will be made private in 5.5")
+	void CacheResourceShadersForRendering(bool bRegenerateId = false);
+	
+	UE_DEPRECATED_FORGAME(5.4, "Do not use. Will be made private in 5.5")
 	void CacheShadersForResources(EShaderPlatform InShaderPlatform, FOpenColorIOTransformResource* InResourcesToCache, bool bApplyCompletedShaderMapForRendering, bool bIsCooking, const ITargetPlatform* TargetPlatform = nullptr);
 
 	UE_DEPRECATED(5.3, "This method is deprecated.")
