OCIO: Prevent render-thread use of game-thread obtained texture resources.

(V3 based on [REDACTED_HANDLE]'s V2 suggested fix.)

#rnx

diff --git a/Engine/Plugins/Compositing/OpenColorIO/Source/OpenColorIO/Private/OpenColorIOColorTransform.cpp b/Engine/Plugins/Compositing/OpenColorIO/Source/OpenColorIO/Private/OpenColorIOColorTransform.cpp
--- a/Engine/Plugins/Compositing/OpenColorIO/Source/OpenColorIO/Private/OpenColorIOColorTransform.cpp
+++ b/Engine/Plugins/Compositing/OpenColorIO/Source/OpenColorIO/Private/OpenColorIOColorTransform.cpp
@@ -496,8 +496,13 @@ FOpenColorIOTransformResource* UOpenColorIOColorTransform::AllocateResource()
 	return new FOpenColorIOTransformResource();
 }
 
-bool UOpenColorIOColorTransform::GetRenderResources(ERHIFeatureLevel::Type InFeatureLevel, FOpenColorIOTransformResource*& OutShaderResource, TSortedMap<int32, FTextureResource*>& OutTextureResources) const
+bool UOpenColorIOColorTransform::GetRenderResources(ERHIFeatureLevel::Type InFeatureLevel, FOpenColorIOTransformResource*& OutShaderResource, TSortedMap<int32, TWeakObjectPtr<UTexture>>& OutTextureResources) const
 {
+	if (!AreRenderResourcesReady())
+	{
+		return false;
+	}
+
 	OutShaderResource = ColorTransformResources[InFeatureLevel];
 	
 	if (OutShaderResource)
@@ -506,18 +511,7 @@ bool UOpenColorIOColorTransform::GetRenderResources(ERHIFeatureLevel::Type InFea
 
 		for (const TPair<int32, TObjectPtr<UTexture>>& Pair : Textures)
 		{
-			if (FTextureResource* TextureResource = Pair.Value->GetResource())
-			{
-				OutTextureResources.Add(Pair.Key, TextureResource);
-			}
-			else
-			{
-				// The texture doesn't exist, we can't get all the resources we need
-				OutShaderResource = nullptr;
-				OutTextureResources.Empty();
-
-				return false;
-			}
+			OutTextureResources.Add(Pair.Key, Pair.Value);
 		}
 	}
 	
diff --git a/Engine/Plugins/Compositing/OpenColorIO/Source/OpenColorIO/Private/OpenColorIOConfiguration.cpp b/Engine/Plugins/Compositing/OpenColorIO/Source/OpenColorIO/Private/OpenColorIOConfiguration.cpp
--- a/Engine/Plugins/Compositing/OpenColorIO/Source/OpenColorIO/Private/OpenColorIOConfiguration.cpp
+++ b/Engine/Plugins/Compositing/OpenColorIO/Source/OpenColorIO/Private/OpenColorIOConfiguration.cpp
@@ -71,7 +71,7 @@ bool UOpenColorIOConfiguration::IsTransformReady(const FOpenColorIOColorConversi
 	return false;
 }
 
-bool UOpenColorIOConfiguration::GetRenderResources(ERHIFeatureLevel::Type InFeatureLevel, const FOpenColorIOColorConversionSettings& InSettings, FOpenColorIOTransformResource*& OutShaderResource, TSortedMap<int32, FTextureResource*>& OutTextureResources)
+bool UOpenColorIOConfiguration::GetRenderResources(ERHIFeatureLevel::Type InFeatureLevel, const FOpenColorIOColorConversionSettings& InSettings, FOpenColorIOTransformResource*& OutShaderResource, TSortedMap<int32, TWeakObjectPtr<UTexture>>& OutTextureResources)
 {
 	TObjectPtr<const UOpenColorIOColorTransform> Transform = FindTransform(InSettings);
 	if (IsValid(Transform))
diff --git a/Engine/Plugins/Compositing/OpenColorIO/Source/OpenColorIO/Private/OpenColorIORendering.cpp b/Engine/Plugins/Compositing/OpenColorIO/Source/OpenColorIO/Private/OpenColorIORendering.cpp
--- a/Engine/Plugins/Compositing/OpenColorIO/Source/OpenColorIO/Private/OpenColorIORendering.cpp
+++ b/Engine/Plugins/Compositing/OpenColorIO/Source/OpenColorIO/Private/OpenColorIORendering.cpp
@@ -43,14 +43,32 @@ void FOpenColorIORendering::AddPass_RenderThread(
 	TShaderMapRef<FScreenPassVS> VertexShader(ShaderMap);
 
 	bool bRenderPassSuccessed = false;
-	if (InPassResource.ShaderResource != nullptr)
+
+	TSortedMap<int32, FTextureResource*> TextureResources;
+	TextureResources.Reserve(InPassResource.TextureResources.Num());
+
+	for (const TPair<int32, TWeakObjectPtr<UTexture>>& Pair : InPassResource.TextureResources)
+	{
+		if (Pair.Value.IsValid(false, true))
+		{
+			// UTexture's GetResource() is safe to call from the render thread.
+			if (FTextureResource* TextureResource = Pair.Value->GetResource())
+			{
+				TextureResources.Add(Pair.Key, TextureResource);
+			}
+		}
+	}
+
+	const bool bAllExpectedTexturesResources = TextureResources.Num() == InPassResource.TextureResources.Num();
+
+	if (bAllExpectedTexturesResources && InPassResource.ShaderResource != nullptr)
 	{
 		TShaderRef<FOpenColorIOPixelShader> OCIOPixelShader = InPassResource.ShaderResource->GetShader<FOpenColorIOPixelShader>();
 
 		FOpenColorIOPixelShaderParameters* Parameters = GraphBuilder.AllocParameters<FOpenColorIOPixelShaderParameters>();
 		Parameters->InputTexture = Input.Texture;
 		Parameters->InputTextureSampler = TStaticSamplerState<>::GetRHI();
-		if (OpenColorIOBindTextureResources(Parameters, InPassResource.TextureResources))
+		if (OpenColorIOBindTextureResources(Parameters, TextureResources))
 		{
 			Parameters->Gamma = InGamma;
 			Parameters->TransformAlpha = (uint32)TransformAlpha;
@@ -173,11 +191,11 @@ bool FOpenColorIORendering::ApplyColorTransform(UWorld* InWorld, const FOpenColo
 
 	const ERHIFeatureLevel::Type FeatureLevel = InWorld->Scene->GetFeatureLevel();
 	FOpenColorIOTransformResource* ShaderResource = nullptr;
-	TSortedMap<int32, FTextureResource*> TransformTextureResources;
+	TSortedMap<int32, TWeakObjectPtr<UTexture>> TransformTextures;
 
 	if (InSettings.ConfigurationSource != nullptr)
 	{
-		const bool bFoundTransform = InSettings.ConfigurationSource->GetRenderResources(FeatureLevel, InSettings, ShaderResource, TransformTextureResources);
+		const bool bFoundTransform = InSettings.ConfigurationSource->GetRenderResources(FeatureLevel, InSettings, ShaderResource, TransformTextures);
 
 		if (bFoundTransform)
 		{
@@ -193,7 +211,7 @@ bool FOpenColorIORendering::ApplyColorTransform(UWorld* InWorld, const FOpenColo
 	}
 	
 	ENQUEUE_RENDER_COMMAND(ProcessColorSpaceTransform)(
-		[FeatureLevel, InputResource, OutputResource, ShaderResource, TextureResources = MoveTemp(TransformTextureResources), TransformName = InSettings.ToString()](FRHICommandListImmediate& RHICmdList)
+		[FeatureLevel, InputResource, OutputResource, ShaderResource, Textures = MoveTemp(TransformTextures), TransformName = InSettings.ToString()](FRHICommandListImmediate& RHICmdList)
 		{
 			FRDGBuilder GraphBuilder(RHICmdList);
 
@@ -208,7 +226,7 @@ bool FOpenColorIORendering::ApplyColorTransform(UWorld* InWorld, const FOpenColo
 				FeatureLevel,
 				FScreenPassTexture(InputTexture),
 				Output,
-				FOpenColorIORenderPassResources{ShaderResource, TextureResources, TransformName},
+				FOpenColorIORenderPassResources{ShaderResource, Textures, TransformName},
 				1.0f); // Set Gamma to 1., since we do not have any display parameters or requirement for Gamma.
 
 			GraphBuilder.Execute();
diff --git a/Engine/Plugins/Compositing/OpenColorIO/Source/OpenColorIO/Public/OpenColorIOColorTransform.h b/Engine/Plugins/Compositing/OpenColorIO/Source/OpenColorIO/Public/OpenColorIOColorTransform.h
--- a/Engine/Plugins/Compositing/OpenColorIO/Source/OpenColorIO/Public/OpenColorIOColorTransform.h
+++ b/Engine/Plugins/Compositing/OpenColorIO/Source/OpenColorIO/Public/OpenColorIOColorTransform.h
@@ -74,7 +74,16 @@ public:
 	/**
 	 * Returns the desired resources required to apply this transform during rendering.
 	 */
-	bool GetRenderResources(ERHIFeatureLevel::Type InFeatureLevel, FOpenColorIOTransformResource*& OutShaderResource, TSortedMap<int32, FTextureResource*>& OutTextureResources) const;
+	UE_DEPRECATED(5.4, "This method is deprecated.")
+	bool GetRenderResources(ERHIFeatureLevel::Type InFeatureLevel, FOpenColorIOTransformResource*& OutShaderResource, TSortedMap<int32, FTextureResource*>& OutTextureResources) const
+	{
+		return false;
+	}
+
+	/**
+	 * Returns the desired resources required to apply this transform during rendering.
+	 */
+	bool GetRenderResources(ERHIFeatureLevel::Type InFeatureLevel, FOpenColorIOTransformResource*& OutShaderResource, TSortedMap<int32, TWeakObjectPtr<UTexture>>& OutTextureResources) const;
 
 	/**
 	 * Returns true if shader/texture resources have finished compiling and are ready for use (to be called on the game thread).
diff --git a/Engine/Plugins/Compositing/OpenColorIO/Source/OpenColorIO/Public/OpenColorIOConfiguration.h b/Engine/Plugins/Compositing/OpenColorIO/Source/OpenColorIO/Public/OpenColorIOConfiguration.h
--- a/Engine/Plugins/Compositing/OpenColorIO/Source/OpenColorIO/Public/OpenColorIOConfiguration.h
+++ b/Engine/Plugins/Compositing/OpenColorIO/Source/OpenColorIO/Public/OpenColorIOConfiguration.h
@@ -18,6 +18,7 @@ class FOpenColorIOWrapperConfig;
 class FOpenColorIOTransformResource;
 class FTextureResource;
 class UOpenColorIOColorTransform;
+class UTexture;
 struct FImageView;
 struct FFileChangeData;
 class SNotificationItem;
@@ -38,7 +39,15 @@ public:
 public:
 
 	bool IsTransformReady(const FOpenColorIOColorConversionSettings& InSettings);
-	bool GetRenderResources(ERHIFeatureLevel::Type InFeatureLevel, const FOpenColorIOColorConversionSettings& InSettings, FOpenColorIOTransformResource*& OutShaderResource, TSortedMap<int32, FTextureResource*>& OutTextureResources);
+	
+	UE_DEPRECATED(5.4, "This method is deprecated.")
+	bool GetRenderResources(ERHIFeatureLevel::Type InFeatureLevel, const FOpenColorIOColorConversionSettings& InSettings, FOpenColorIOTransformResource*& OutShaderResource, TSortedMap<int32, FTextureResource*>& OutTextureResources)
+	{
+		return false;
+	}
+
+	bool GetRenderResources(ERHIFeatureLevel::Type InFeatureLevel, const FOpenColorIOColorConversionSettings& InSettings, FOpenColorIOTransformResource*& OutShaderResource, TSortedMap<int32, TWeakObjectPtr<UTexture>>& OutTextureResources);
+	
 	bool HasTransform(const FString& InSourceColorSpace, const FString& InDestinationColorSpace);
 	bool HasTransform(const FString& InSourceColorSpace, const FString& InDisplay, const FString& InView, EOpenColorIOViewTransformDirection InDirection);
 	bool HasDesiredColorSpace(const FOpenColorIOColorSpace& ColorSpace) const;
diff --git a/Engine/Plugins/Compositing/OpenColorIO/Source/OpenColorIO/Public/OpenColorIORendering.h b/Engine/Plugins/Compositing/OpenColorIO/Source/OpenColorIO/Public/OpenColorIORendering.h
--- a/Engine/Plugins/Compositing/OpenColorIO/Source/OpenColorIO/Public/OpenColorIORendering.h
+++ b/Engine/Plugins/Compositing/OpenColorIO/Source/OpenColorIO/Public/OpenColorIORendering.h
@@ -3,6 +3,7 @@
 #pragma once
 
 #include "CoreMinimal.h"
+#include "UObject/WeakObjectPtr.h"
 #include "RHIFwd.h"
 
 struct FOpenColorIOColorConversionSettings;
@@ -25,7 +26,7 @@ struct OPENCOLORIO_API FOpenColorIORenderPassResources
 	FOpenColorIOTransformResource* ShaderResource = nullptr;
 	
 	/** Collection of LUT textures needed by the shader. */
-	TSortedMap<int32, FTextureResource*> TextureResources = {};
+	TSortedMap<int32, TWeakObjectPtr<UTexture>> TextureResources = {};
 
 	/** Color transform string description. */
 	FString TransformName = FString();
diff --git a/Engine/Plugins/Media/MediaIOFramework/Source/MediaIOCore/Private/Shared/MediaIOCoreTextureSampleBase.cpp b/Engine/Plugins/Media/MediaIOFramework/Source/MediaIOCore/Private/Shared/MediaIOCoreTextureSampleBase.cpp
--- a/Engine/Plugins/Media/MediaIOFramework/Source/MediaIOCore/Private/Shared/MediaIOCoreTextureSampleBase.cpp
+++ b/Engine/Plugins/Media/MediaIOFramework/Source/MediaIOCore/Private/Shared/MediaIOCoreTextureSampleBase.cpp
@@ -218,7 +218,7 @@ bool FMediaIOCoreTextureSampleBase::ApplyColorConversion(FTexture2DRHIRef& InSrc
 			CachedOCIOResources = MakeShared<FOpenColorIORenderPassResources>();
 
 			FOpenColorIOTransformResource* ShaderResource = nullptr;
-			TSortedMap<int32, FTextureResource*> TransformTextureResources;
+			TSortedMap<int32, TWeakObjectPtr<UTexture>> TransformTextureResources;
 
 			if (ColorConversionSettings->ConfigurationSource != nullptr)
 			{
