OCIO: Add context to configuration hash.

We also take this opportunity to cleanup the hash computation, and use xxhash instead to match the library.

#rnx

diff --git a/Engine/Plugins/Compositing/OpenColorIO/Source/OpenColorIO/Private/OpenColorIOConfiguration.cpp b/Engine/Plugins/Compositing/OpenColorIO/Source/OpenColorIO/Private/OpenColorIOConfiguration.cpp
--- a/Engine/Plugins/Compositing/OpenColorIO/Source/OpenColorIO/Private/OpenColorIOConfiguration.cpp
+++ b/Engine/Plugins/Compositing/OpenColorIO/Source/OpenColorIO/Private/OpenColorIOConfiguration.cpp
@@ -7,6 +7,7 @@
 #include "EngineAnalytics.h"
 #include "Engine/VolumeTexture.h"
 #include "Framework/Notifications/NotificationManager.h"
+#include "Hash/xxhash.h"
 #include "Math/PackedVector.h"
 #include "Misc/PathViews.h"
 #include "Modules/ModuleManager.h"
@@ -177,37 +178,16 @@ void UOpenColorIOConfiguration::ReloadExistingColorspaces(bool bForce)
 #if WITH_OCIO
 	if (Config && Config->IsValid())
 	{
-		FString LoadedConfigHash = Config->GetCacheID();
-		if (LoadedConfigHash.IsEmpty())
+		FString NewConfigHash;
+		if (!GetHash(NewConfigHash))
 		{
-			UE_LOG(LogOpenColorIO, Warning, TEXT("Failed to get cache ID: forcing constant transform recreation. Please fix invalid config."));
 			bForce = true;
 		}
 
-		LoadedConfigHash += FString(OpenColorIOWrapper::GetVersion());
-		
-		const UOpenColorIOSettings* Settings = GetDefault<UOpenColorIOSettings>();
-		if (Settings->bSupportInverseViewTransforms)
-		{
-			LoadedConfigHash += FString(TEXT("_Inv"));
-		}
-
-		const UE::Color::FColorSpace& WCS = UE::Color::FColorSpace::GetWorking();
-		if (!WCS.IsSRGB())
-		{
-			// The working color space is uniquely defined by its chromaticities (as loaded from renderer settings).
-			uint32 WCSHash = 0;
-			WCSHash ^= GetTypeHash(WCS.GetRedChromaticity());
-			WCSHash ^= GetTypeHash(WCS.GetGreenChromaticity());
-			WCSHash ^= GetTypeHash(WCS.GetBlueChromaticity());
-			WCSHash ^= GetTypeHash(WCS.GetWhiteChromaticity());
-			LoadedConfigHash += FString::Printf(TEXT("_WCS-%u"), WCSHash);
-		}
-
 		// Hash is different, proceed with the regeneration...
-		if (ConfigHash != LoadedConfigHash || bForce)
+		if (ConfigHash != NewConfigHash || bForce)
 		{
-			ConfigHash = LoadedConfigHash;
+			ConfigHash = NewConfigHash;
 
 			TArray<FOpenColorIOColorSpace> ColorSpacesToBeReloaded = DesiredColorSpaces;
 			TArray<FOpenColorIODisplayView> DisplayViewsToBeReloaded = DesiredDisplayViews;
@@ -254,6 +234,8 @@ void UOpenColorIOConfiguration::ReloadExistingColorspaces(bool bForce)
 				}
 			}
 
+			const UOpenColorIOSettings* Settings = GetDefault<UOpenColorIOSettings>();
+
 			// Genereate new shaders.
 			for (int32 indexTop = 0; indexTop < DesiredColorSpaces.Num(); ++indexTop)
 			{
@@ -288,6 +270,7 @@ void UOpenColorIOConfiguration::ReloadExistingColorspaces(bool bForce)
 		DesiredColorSpaces.Reset();
 		DesiredDisplayViews.Reset();
 		ColorTransforms.Reset();
+		ConfigHash.Empty();
 	}
 
 	if (!MarkPackageDirty())
@@ -660,6 +643,63 @@ void UOpenColorIOConfiguration::LoadConfiguration()
 #endif //WITH_OCIO
 }
 
+bool UOpenColorIOConfiguration::GetHash(FString& OutHash) const
+{
+#if WITH_OCIO
+	if (Config && Config->IsValid())
+	{
+		const FString CacheID = Config->GetCacheID();
+		if (CacheID.IsEmpty())
+		{
+			UE_LOG(LogOpenColorIO, Warning, TEXT("Failed to get cache ID: forcing constant transform recreation. Please fix invalid config."));
+			return false;
+		}
+
+		const FStringView CacheIDView = CacheID;
+		const uint32 VersionID = OpenColorIOWrapper::GetVersionHex();
+		const uint8 bSupportsInverse = GetDefault<UOpenColorIOSettings>()->bSupportInverseViewTransforms;
+
+		// Note that we use 128-bit xxhash to match the library.
+		// See: [REDACTED_PROFILE_URL]/OpenColorIO/wiki/Caches-and-cache-IDs-in-OpenColorIO#object-cacheids
+
+		FXxHash128Builder HashBuilder;
+		HashBuilder.Update(CacheIDView.GetData(), CacheIDView.Len() * sizeof(CacheIDView[0]));
+		HashBuilder.Update(&VersionID, sizeof(VersionID));
+		HashBuilder.Update(&bSupportsInverse, sizeof(bSupportsInverse));
+
+		const UE::Color::FColorSpace& WCS = UE::Color::FColorSpace::GetWorking();
+		if (!WCS.IsSRGB())
+		{
+			// The working color space is uniquely defined by its chromaticities (as loaded from renderer settings).
+			HashBuilder.Update(&WCS.GetRedChromaticity(), sizeof(FVector2d));
+			HashBuilder.Update(&WCS.GetGreenChromaticity(), sizeof(FVector2d));
+			HashBuilder.Update(&WCS.GetBlueChromaticity(), sizeof(FVector2d));
+			HashBuilder.Update(&WCS.GetWhiteChromaticity(), sizeof(FVector2d));
+		}
+
+		for (const TPair<FString, FString>& Pair : Context)
+		{
+			if (!Pair.Value.IsEmpty())
+			{
+				const FStringView Key = Pair.Key;
+				const FStringView Value = Pair.Value;
+
+				HashBuilder.Update(Key.GetData(), Key.Len() * sizeof(Key[0]));
+				HashBuilder.Update(Value.GetData(), Value.Len() * sizeof(Value[0]));
+			}
+		}
+		
+		FStringBuilderBase StringBuilder;
+		StringBuilder << HashBuilder.Finalize();
+		OutHash = StringBuilder.ToString();
+
+		return true;
+	}
+#endif // WITH_OCIO
+
+	return false;
+}
+
 #if WITH_EDITOR
 
 void UOpenColorIOConfiguration::PostEditChangeProperty(FPropertyChangedEvent& PropertyChangedEvent)
@@ -727,6 +767,12 @@ void UOpenColorIOConfiguration::PostEditChangeProperty(FPropertyChangedEvent& Pr
 			CleanupTransforms();
 		}
 	}
+	else if (PropertyChangedEvent.GetPropertyName() == GET_MEMBER_NAME_CHECKED(UOpenColorIOConfiguration, Context))
+	{
+		// Note: Reload calls LoadConfiguration() internally.
+		ReloadExistingColorspaces();
+	}
+
 	Super::PostEditChangeProperty(PropertyChangedEvent);
 }
 
diff --git a/Engine/Plugins/Compositing/OpenColorIO/Source/OpenColorIO/Public/OpenColorIOConfiguration.h b/Engine/Plugins/Compositing/OpenColorIO/Source/OpenColorIO/Public/OpenColorIOConfiguration.h
--- a/Engine/Plugins/Compositing/OpenColorIO/Source/OpenColorIO/Public/OpenColorIOConfiguration.h
+++ b/Engine/Plugins/Compositing/OpenColorIO/Source/OpenColorIO/Public/OpenColorIOConfiguration.h
@@ -125,6 +125,9 @@ private:
 	/** Load the config file to initialize the configuration wrapper object. Automatically called internally. */
 	void LoadConfiguration();
 
+	/** Obtain the hash of the config and additional dependent state such as library version, working color space and context. */
+	bool GetHash(FString& OutHash) const;
+
 #if WITH_EDITOR
 	/** This method resets the status of Notification dialog and reacts depending on user's choice. */
 	void OnToastCallback(bool bInReloadColorspaces);
