OCIO: Compilation time improvements.

diff --git a/Engine/Plugins/Compositing/OpenColorIO/Source/OpenColorIO/Public/OpenColorIOColorSpace.h b/Engine/Plugins/Compositing/OpenColorIO/Source/OpenColorIO/Public/OpenColorIOColorSpace.h
--- a/Engine/Plugins/Compositing/OpenColorIO/Source/OpenColorIO/Public/OpenColorIOColorSpace.h
+++ b/Engine/Plugins/Compositing/OpenColorIO/Source/OpenColorIO/Public/OpenColorIOColorSpace.h
@@ -2,7 +2,9 @@
 
 #pragma once
 
+#if UE_ENABLE_INCLUDE_ORDER_DEPRECATED_IN_5_4
 #include "CoreMinimal.h"
+#endif
 
 #include "OpenColorIOColorSpace.generated.h"
 
diff --git a/Engine/Plugins/Compositing/OpenColorIO/Source/OpenColorIO/Public/OpenColorIOShader.h b/Engine/Plugins/Compositing/OpenColorIO/Source/OpenColorIO/Public/OpenColorIOShader.h
--- a/Engine/Plugins/Compositing/OpenColorIO/Source/OpenColorIO/Public/OpenColorIOShader.h
+++ b/Engine/Plugins/Compositing/OpenColorIO/Source/OpenColorIO/Public/OpenColorIOShader.h
@@ -7,7 +7,9 @@
 #pragma once
 
 #include "Containers/ContainersFwd.h"
+#if UE_ENABLE_INCLUDE_ORDER_DEPRECATED_IN_5_4
 #include "GlobalShader.h"
+#endif
 #include "OpenColorIOShaderType.h"
 #include "OpenColorIOShared.h"
 #include "ShaderParameters.h"
diff --git a/Engine/Plugins/Compositing/OpenColorIO/Source/OpenColorIO/Public/OpenColorIOShared.h b/Engine/Plugins/Compositing/OpenColorIO/Source/OpenColorIO/Public/OpenColorIOShared.h
--- a/Engine/Plugins/Compositing/OpenColorIO/Source/OpenColorIO/Public/OpenColorIOShared.h
+++ b/Engine/Plugins/Compositing/OpenColorIO/Source/OpenColorIO/Public/OpenColorIOShared.h
@@ -6,23 +6,28 @@
 
 #pragma once
 
-#include "CoreMinimal.h"
-#include "Engine/EngineTypes.h"
+#include "CoreTypes.h"
 #include "Misc/Guid.h"
-#include "Misc/Optional.h"
 #include "Misc/SecureHash.h"
 #include "RenderResource.h"
 #if UE_ENABLE_INCLUDE_ORDER_DEPRECATED_IN_5_2
 #include "RenderingThread.h"
 #endif
 #include "RenderDeferredCleanup.h"
-#include "RHI.h"
+#include "RHIFeatureLevel.h"
 #include "SceneTypes.h"
 #include "Shader.h"
-#include "ShaderCompiler.h"
 #include "StaticParameterSet.h"
 #include "Templates/RefCounting.h"
+
+#if UE_ENABLE_INCLUDE_ORDER_DEPRECATED_IN_5_4
+#include "CoreMinimal.h"
+#include "Engine/EngineTypes.h"
+#include "Misc/Optional.h"
+#include "ShaderCompiler.h"
 #include "OpenColorIOShaderCompilationManager.h"
+#include "RHI.h"
+#endif
 
 /** Enum used to indicate whether the working color space should be used as a source or destination. */
 enum class EOpenColorIOWorkingColorSpaceTransform : uint8
@@ -187,7 +192,7 @@ public:
 		);
 
 	/** Sorts the incoming compiled jobs into the appropriate OCIO shader maps, and finalizes this shader map so that it can be used for rendering. */
-	bool ProcessCompilationResults(const TArray<FShaderCommonCompileJobPtr>& InCompilationResults, int32& InOutResultIndex, float& InOutTimeBudget);
+	bool ProcessCompilationResults(const TArray<TRefCountPtr<class FShaderCommonCompileJob>>& InCompilationResults, int32& InOutResultIndex, float& InOutTimeBudget);
 #endif // WITH_EDITOR
 
 	/**
