OCIO: Prevent use after free of display-view.

diff --git a/Engine/Plugins/Compositing/OpenColorIO/Source/OpenColorIO/Private/OpenColorIOColorSpace.cpp b/Engine/Plugins/Compositing/OpenColorIO/Source/OpenColorIO/Private/OpenColorIOColorSpace.cpp
--- a/Engine/Plugins/Compositing/OpenColorIO/Source/OpenColorIO/Private/OpenColorIOColorSpace.cpp
+++ b/Engine/Plugins/Compositing/OpenColorIO/Source/OpenColorIO/Private/OpenColorIOColorSpace.cpp
@@ -87,7 +87,7 @@ FString FOpenColorIODisplayView::ToString() const
 {
 	if (IsValid())
 	{
-		return Display + TEXT(" - ") + View;
+		return FString::Printf(TEXT("%s - %s"), *Display, *View);
 	}
 
 	return {};
diff --git a/Engine/Plugins/Compositing/OpenColorIO/Source/OpenColorIOEditor/Private/OpenColorIOColorSpaceCustomization.cpp b/Engine/Plugins/Compositing/OpenColorIO/Source/OpenColorIOEditor/Private/OpenColorIOColorSpaceCustomization.cpp
--- a/Engine/Plugins/Compositing/OpenColorIO/Source/OpenColorIOEditor/Private/OpenColorIOColorSpaceCustomization.cpp
+++ b/Engine/Plugins/Compositing/OpenColorIO/Source/OpenColorIOEditor/Private/OpenColorIOColorSpaceCustomization.cpp
@@ -81,11 +81,13 @@ void FOpenColorIODisplayViewCustomization::CustomizeHeader(TSharedRef<IPropertyH
 
 		TArray<void*> RawData;
 		CachedProperty->AccessRawData(RawData);
-
 		check(RawData.Num() == 1);
-		FOpenColorIODisplayView* DisplayViewValue = reinterpret_cast<FOpenColorIODisplayView*>(RawData[0]);
-
+		
+		const FOpenColorIODisplayView* DisplayViewValue = reinterpret_cast<const FOpenColorIODisplayView*>(RawData[0]);
 		check(DisplayViewValue);
+
+		const FString DisplayViewName = DisplayViewValue->ToString();
+		
 		TSharedPtr<IPropertyUtilities> PropertyUtils = CustomizationUtils.GetPropertyUtilities();
 
 		HeaderRow
@@ -102,9 +104,8 @@ void FOpenColorIODisplayViewCustomization::CustomizeHeader(TSharedRef<IPropertyH
 				.VAlign(VAlign_Center)
 				[
 					SNew(STextBlock)
-					.Text(MakeAttributeLambda([DisplayViewValue] {
-							const FString DisplayViewName = DisplayViewValue->ToString();
-							
+					.Text(MakeAttributeLambda([DisplayViewName]
+						{
 							if(!DisplayViewName.IsEmpty())
 							{
 								return FText::FromString(DisplayViewName);
