OCIO: Cleanup of deprecated functions.

#rnx

diff --git a/Engine/Plugins/Compositing/OpenColorIO/Source/OpenColorIO/Private/OpenColorIOColorTransform.cpp b/Engine/Plugins/Compositing/OpenColorIO/Source/OpenColorIO/Private/OpenColorIOColorTransform.cpp
--- a/Engine/Plugins/Compositing/OpenColorIO/Source/OpenColorIO/Private/OpenColorIOColorTransform.cpp
+++ b/Engine/Plugins/Compositing/OpenColorIO/Source/OpenColorIO/Private/OpenColorIOColorTransform.cpp
@@ -139,16 +139,6 @@ UOpenColorIOColorTransform::UOpenColorIOColorTransform(const FObjectInitializer&
 {
 }
 
-bool UOpenColorIOColorTransform::Initialize(UOpenColorIOConfiguration* InOwner, const FString& InSourceColorSpace, const FString& InDestinationColorSpace, const TMap<FString, FString>& /*InContextKeyValues*/)
-{
-	return Initialize(InSourceColorSpace, InDestinationColorSpace);
-}
-
-bool UOpenColorIOColorTransform::Initialize(UOpenColorIOConfiguration* InOwner, const FString& InSourceColorSpace, const FString& InDisplay, const FString& InView, EOpenColorIOViewTransformDirection InDirection, const TMap<FString, FString>& /*InContextKeyValues*/)
-{
-	return Initialize(InSourceColorSpace, InDisplay, InView, InDirection);
-}
-
 bool UOpenColorIOColorTransform::Initialize(const FString& InSourceColorSpace, const FString& InDestinationColorSpace, const TMap<FString, FString>& /*InContextKeyValues*/)
 {
 	return Initialize(InSourceColorSpace, InDestinationColorSpace);
@@ -491,11 +481,6 @@ void UOpenColorIOColorTransform::CacheShadersForResources(EShaderPlatform InShad
 	}
 }
 
-FOpenColorIOTransformResource* UOpenColorIOColorTransform::AllocateResource()
-{
-	return new FOpenColorIOTransformResource();
-}
-
 bool UOpenColorIOColorTransform::GetRenderResources(ERHIFeatureLevel::Type InFeatureLevel, FOpenColorIOTransformResource*& OutShaderResource, TSortedMap<int32, TWeakObjectPtr<UTexture>>& OutTextureResources) const
 {
 	if (!AreRenderResourcesReady())
diff --git a/Engine/Plugins/Compositing/OpenColorIO/Source/OpenColorIO/Private/OpenColorIOConfiguration.cpp b/Engine/Plugins/Compositing/OpenColorIO/Source/OpenColorIO/Private/OpenColorIOConfiguration.cpp
--- a/Engine/Plugins/Compositing/OpenColorIO/Source/OpenColorIO/Private/OpenColorIOConfiguration.cpp
+++ b/Engine/Plugins/Compositing/OpenColorIO/Source/OpenColorIO/Private/OpenColorIOConfiguration.cpp
@@ -793,46 +793,6 @@ void UOpenColorIOConfiguration::OnToastCallback(bool bInReloadColorspaces)
 	}
 }
 
-
-FOpenColorIOEditorConfigurationInspector::FOpenColorIOEditorConfigurationInspector(const UOpenColorIOConfiguration& InConfiguration)
-	: Configuration(InConfiguration)
-{
-}
-
-int32 FOpenColorIOEditorConfigurationInspector::GetNumColorSpaces() const
-{
-	return Configuration.GetConfigWrapper()->GetNumColorSpaces();
-}
-
-FString FOpenColorIOEditorConfigurationInspector::GetColorSpaceName(int32 Index) const
-{
-	return Configuration.GetConfigWrapper()->GetColorSpaceName(Index);
-}
-
-FString FOpenColorIOEditorConfigurationInspector::GetColorSpaceFamilyName(const TCHAR* InColorSpaceName) const
-{
-	return Configuration.GetConfigWrapper()->GetColorSpaceFamilyName(InColorSpaceName);
-}
-
-int32 FOpenColorIOEditorConfigurationInspector::GetNumDisplays() const
-{
-	return Configuration.GetConfigWrapper()->GetNumDisplays();
-}
-
-FString FOpenColorIOEditorConfigurationInspector::GetDisplayName(int32 Index) const
-{
-	return Configuration.GetConfigWrapper()->GetDisplayName(Index);
-}
-
-int32 FOpenColorIOEditorConfigurationInspector::GetNumViews(const TCHAR* InDisplayName) const
-{
-	return Configuration.GetConfigWrapper()->GetNumViews(InDisplayName);
-}
-
-FString FOpenColorIOEditorConfigurationInspector::GetViewName(const TCHAR* InDisplayName, int32 Index) const
-{
-	return Configuration.GetConfigWrapper()->GetViewName(InDisplayName, Index);
-}
 #endif //WITH_EDITOR
 
 
diff --git a/Engine/Plugins/Compositing/OpenColorIO/Source/OpenColorIO/Public/IOpenColorIOModule.h b/Engine/Plugins/Compositing/OpenColorIO/Source/OpenColorIO/Public/IOpenColorIOModule.h
--- a/Engine/Plugins/Compositing/OpenColorIO/Source/OpenColorIO/Public/IOpenColorIOModule.h
+++ b/Engine/Plugins/Compositing/OpenColorIO/Source/OpenColorIO/Public/IOpenColorIOModule.h
@@ -44,7 +44,4 @@ public:
 	 * Returns the collections of ocio display look currently in play
 	 */
 	virtual FOpenColorIODisplayManager& GetDisplayManager() = 0;
-
-	UE_DEPRECATED(5.3, "This function is deprecated and has been replaced by FOpenColorIOWrapperConfig class in the OpenColorIOWrapper module.")
-	FOpenColorIONativeInterchangeConfiguration* GetNativeInterchangeConfig_Internal() { return nullptr; };
 };
\ No newline at end of file
diff --git a/Engine/Plugins/Compositing/OpenColorIO/Source/OpenColorIO/Public/OpenColorIOColorSpace.h b/Engine/Plugins/Compositing/OpenColorIO/Source/OpenColorIO/Public/OpenColorIOColorSpace.h
--- a/Engine/Plugins/Compositing/OpenColorIO/Source/OpenColorIO/Public/OpenColorIOColorSpace.h
+++ b/Engine/Plugins/Compositing/OpenColorIO/Source/OpenColorIO/Public/OpenColorIOColorSpace.h
@@ -150,9 +150,6 @@ public:
 	UPROPERTY(EditAnywhere, BlueprintReadWrite, Category = ColorSpace)
 	EOpenColorIOViewTransformDirection DisplayViewDirection = EOpenColorIOViewTransformDirection::Forward;
 
-	/** Delegate triggered upon changes to the settings. */
-	UE_DEPRECATED(5.3, "This delegate is deprecated.")
-	FOnConversionSettingsChange& OnConversionSettingsChanged() { return ConversionSettingsChanged; }
 public:
 
 	void PostSerialize(const FArchive& Ar);
@@ -200,10 +197,6 @@ public:
 	
 	/** Whether or not these settings are of the display-view type. */
 	bool IsDisplayView() const;
-
-private:
-
-	FOnConversionSettingsChange ConversionSettingsChanged;
 };
 
 template<> struct TStructOpsTypeTraits<FOpenColorIOColorConversionSettings> : public TStructOpsTypeTraitsBase2<FOpenColorIOColorConversionSettings>
diff --git a/Engine/Plugins/Compositing/OpenColorIO/Source/OpenColorIO/Public/OpenColorIOColorTransform.h b/Engine/Plugins/Compositing/OpenColorIO/Source/OpenColorIO/Public/OpenColorIOColorTransform.h
--- a/Engine/Plugins/Compositing/OpenColorIO/Source/OpenColorIO/Public/OpenColorIOColorTransform.h
+++ b/Engine/Plugins/Compositing/OpenColorIO/Source/OpenColorIO/Public/OpenColorIOColorTransform.h
@@ -34,10 +34,6 @@ public:
 	virtual ~UOpenColorIOColorTransform() {};
 
 public:
-	UE_DEPRECATED(5.3, "This method is deprecated, please use Initialize without the owner argument.")
-	bool Initialize(UOpenColorIOConfiguration* InOwner, const FString& InSourceColorSpace, const FString& InDestinationColorSpace, const TMap<FString, FString>& InContextKeyValues = {});
-	UE_DEPRECATED(5.3, "This method is deprecated, please use Initialize without the owner argument.")
-	bool Initialize(UOpenColorIOConfiguration* InOwner, const FString& InSourceColorSpace, const FString& InDisplay, const FString& InView, EOpenColorIOViewTransformDirection InDirection, const TMap<FString, FString>& InContextKeyValues = {});
 	UE_DEPRECATED(5.4, "This method is deprecated, please use Initialize without the context argument.")
 	bool Initialize(const FString& InSourceColorSpace, const FString& InDestinationColorSpace, const TMap<FString, FString>& InContextKeyValues);
 	UE_DEPRECATED(5.4, "This method is deprecated, please use Initialize without the context argument.")
@@ -49,15 +45,6 @@ public:
 	/** Initialize resources for display-view transform. */
 	bool Initialize(const FString& InSourceColorSpace, const FString& InDisplay, const FString& InView, EOpenColorIOViewTransformDirection InDirection);
 
-	/**
-	 * Serialize LUT data. This will effectively serialize the LUT only when cooking
-	 */
-	UE_DEPRECATED(5.3, "This method is deprecated.")
-	void SerializeLuts(FArchive& Ar) { }
-
-	UE_DEPRECATED(5.3, "This method is deprecated.")
-	void CacheResourceTextures() { }
-
 	/**
 	 * Cache resource shaders for rendering.
 	 * If a matching shader map is not found in memory or the DDC, a new one will be compiled.
@@ -68,9 +55,6 @@ public:
 	UE_DEPRECATED_FORGAME(5.4, "Do not use. Will be made private in 5.5")
 	void CacheShadersForResources(EShaderPlatform InShaderPlatform, FOpenColorIOTransformResource* InResourcesToCache, bool bApplyCompletedShaderMapForRendering, bool bIsCooking, const ITargetPlatform* TargetPlatform = nullptr);
 
-	UE_DEPRECATED(5.3, "This method is deprecated.")
-	FOpenColorIOTransformResource* AllocateResource();
-
 	/**
 	 * Returns the desired resources required to apply this transform during rendering.
 	 */
@@ -151,12 +135,6 @@ protected:
 	UE_DEPRECATED_FORGAME(5.4, "Do not use. Will be made private in 5.5")
 	static void GetOpenColorIOLUTKeyGuid(const FString& InProcessorIdentifier, const FName& InName, FGuid& OutLutGuid );
 
-	UE_DEPRECATED(5.3, "This method is deprecated.")
-	bool GenerateColorTransformData(const FString& InSourceColorSpace, const FString& InDestinationColorSpace) { return false; }
-
-	UE_DEPRECATED(5.3, "This method is deprecated.")
-	bool GenerateColorTransformData(const FString& InSourceColorSpace, const FString& InDisplay, const FString& InView, EOpenColorIOViewTransformDirection InDirection) { return false; }
-
 	/**
 	 * Helper function returning the color space transform name based on source and destination color spaces.
 	 */
@@ -233,12 +211,6 @@ public:
 
 public:
 
-#if WITH_EDITORONLY_DATA
-	UE_DEPRECATED(5.3, "ConfigurationOwner is deprecated, use GetOuter() instead.")
-	UPROPERTY(Transient, meta = (DeprecatedProperty))
-	TObjectPtr<UOpenColorIOConfiguration> ConfigurationOwner_DEPRECATED;
-#endif
-
 	UPROPERTY(VisibleAnywhere, Category = "ColorSpace")
 	bool bIsDisplayViewType = false;
 
diff --git a/Engine/Plugins/Compositing/OpenColorIO/Source/OpenColorIO/Public/OpenColorIOConfiguration.h b/Engine/Plugins/Compositing/OpenColorIO/Source/OpenColorIO/Public/OpenColorIOConfiguration.h
--- a/Engine/Plugins/Compositing/OpenColorIO/Source/OpenColorIO/Public/OpenColorIOConfiguration.h
+++ b/Engine/Plugins/Compositing/OpenColorIO/Source/OpenColorIO/Public/OpenColorIOConfiguration.h
@@ -179,37 +179,3 @@ private:
 	FString ConfigHash;
 };
 
-
-#if WITH_EDITOR
-class FOpenColorIOEditorConfigurationInspector {
-public:
-
-	UE_DEPRECATED(5.3, "This class is deprecated and has been replaced by FOpenColorIOWrapperConfig class in the OpenColorIOWrapper module.")
-	OPENCOLORIO_API FOpenColorIOEditorConfigurationInspector(const UOpenColorIOConfiguration& InConfiguration);
-
-	UE_DEPRECATED(5.3, "This class is deprecated and has been replaced by FOpenColorIOWrapperConfig class in the OpenColorIOWrapper module.")
-	OPENCOLORIO_API int32 GetNumColorSpaces() const;
-	
-	UE_DEPRECATED(5.3, "This class is deprecated and has been replaced by FOpenColorIOWrapperConfig class in the OpenColorIOWrapper module.")
-	OPENCOLORIO_API FString GetColorSpaceName(int32 Index) const;
-	
-	UE_DEPRECATED(5.3, "This class is deprecated and has been replaced by FOpenColorIOWrapperConfig class in the OpenColorIOWrapper module.")
-	OPENCOLORIO_API FString GetColorSpaceFamilyName(const TCHAR* InColorSpaceName) const;
-
-	UE_DEPRECATED(5.3, "This class is deprecated and has been replaced by FOpenColorIOWrapperConfig class in the OpenColorIOWrapper module.")
-	OPENCOLORIO_API int32 GetNumDisplays() const;
-	
-	UE_DEPRECATED(5.3, "This class is deprecated and has been replaced by FOpenColorIOWrapperConfig class in the OpenColorIOWrapper module.")
-	OPENCOLORIO_API FString GetDisplayName(int32 Index) const;
-
-	UE_DEPRECATED(5.3, "This class is deprecated and has been replaced by FOpenColorIOWrapperConfig class in the OpenColorIOWrapper module.")
-	OPENCOLORIO_API int32 GetNumViews(const TCHAR* InDisplayName) const;
-	
-	UE_DEPRECATED(5.3, "This class is deprecated and has been replaced by FOpenColorIOWrapperConfig class in the OpenColorIOWrapper module.")
-	OPENCOLORIO_API FString GetViewName(const TCHAR* InDisplayName, int32 Index) const;
-
-private:
-
-	const UOpenColorIOConfiguration& Configuration;
-};
-#endif //WITH_EDITOR
