Use FindOrAddSection to simplify some menu code

diff --git a/Engine/Plugins/Compositing/OpenColorIO/Source/OpenColorIOEditor/Private/OpenColorIOEditorModule.cpp b/Engine/Plugins/Compositing/OpenColorIO/Source/OpenColorIOEditor/Private/OpenColorIOEditorModule.cpp
--- a/Engine/Plugins/Compositing/OpenColorIO/Source/OpenColorIOEditor/Private/OpenColorIOEditorModule.cpp
+++ b/Engine/Plugins/Compositing/OpenColorIO/Source/OpenColorIOEditor/Private/OpenColorIOEditorModule.cpp
@@ -143,21 +143,13 @@ void FOpenColorIOEditorModule::RegisterViewMenuExtension()
 	FToolMenuOwnerScoped ToolMenuOwnerScoped(this);
 
 	UToolMenu* Menu = UToolMenus::Get()->ExtendMenu("LevelEditor.LevelViewportToolBar.View");
-	const FName ColorManagementSection = "ColorManagement";
-	FToolMenuSection* Section = Menu->FindSection(ColorManagementSection);
-	if (Section == nullptr)
-	{
-		//Section not found, create one with the label
-		Section = &Menu->AddSection(ColorManagementSection, LOCTEXT("ColorManagement_Label", "Color Management"));
-	}
-
-	check(Section);
-
-	Section->AddSubMenu(
-		"OCIODisplaySubMenu",
+	FToolMenuSection& Section = Menu->FindOrAddSection(
+		"ColorManagement", LOCTEXT("ColorManagement_Label", "Color Management"));
+	Section.AddSubMenu("OCIODisplaySubMenu",
 		LOCTEXT("OCIODisplaySubMenu_Label", "OCIO Display"),
 		LOCTEXT("OCIODisplaySubMenu_ToolTip", "Configure the viewport to use an OCIO display configuration"),
-		FNewToolMenuDelegate::CreateRaw(this, &FOpenColorIOEditorModule::AddOpenColorIODisplaySubMenu), false);
+		FNewToolMenuDelegate::CreateRaw(this, &FOpenColorIOEditorModule::AddOpenColorIODisplaySubMenu),
+		false);
 }
 
 void FOpenColorIOEditorModule::UnregisterViewMenuExtension()
