Prevent OpenColorIOEditor from crashing when it adds it menu entry if the viewport is null.

This can happen if the ToolMenu is traversed very early at startup, which the new viewport toolbar does to look for top-level flagged menu entries.

diff --git a/Engine/Plugins/Compositing/OpenColorIO/Source/OpenColorIOEditor/Private/OpenColorIOEditorModule.cpp b/Engine/Plugins/Compositing/OpenColorIO/Source/OpenColorIOEditor/Private/OpenColorIOEditorModule.cpp
--- a/Engine/Plugins/Compositing/OpenColorIO/Source/OpenColorIOEditor/Private/OpenColorIOEditorModule.cpp
+++ b/Engine/Plugins/Compositing/OpenColorIO/Source/OpenColorIOEditor/Private/OpenColorIOEditorModule.cpp
@@ -168,26 +168,37 @@ void FOpenColorIOEditorModule::AddOpenColorIODisplaySubMenu(UToolMenu* Menu)
 	// menu-in-a-menu which leads to double search fields if the parent is searchable too.
 	Menu->bSearchable = false;
 
-	//Viewport menu was clicked, get which one was hit
-	FViewport* CurrentViewport = GEditor->GetActiveViewport();
+	FToolMenuSection& Section = Menu->AddDynamicSection("DynamicDisplayConfiguration",
+		FNewToolMenuDelegate::CreateLambda([this](UToolMenu* DynamicSectionMenu) -> void {
+			// Viewport menu was clicked, get which one was hit
+			FViewport* CurrentViewport = GEditor->GetActiveViewport();
 
-	//Make sure we know about that viewport
-	TrackNewViewportIfRequired(CurrentViewport);
-	
-	//Fetch configuration for this viewport. If none were made, we'll populate UI with default values
-	const FOpenColorIODisplayConfiguration& Configuration = IOpenColorIOModule::Get().GetDisplayManager().FindOrAddDisplayConfiguration(CurrentViewport->GetClient());
-
-	//Add OCIO display section
-	FToolMenuSection& Section = Menu->AddSection("DisplayConfiguration", LOCTEXT("DisplayConfiguration_Label", "Display Configuration"));
-	Section.AddEntry(FToolMenuEntry::InitWidget(
-		"DisplayConfigurationWidget"
-		, SNew(SOpenColorIODisplay)
-			.Viewport(CurrentViewport)
-			.InitialConfiguration(Configuration)
-			.OnConfigurationChanged(FOnDisplayConfigurationChanged::CreateRaw(this, &FOpenColorIOEditorModule::OnDisplayConfigurationChanged))
-		,FText::GetEmpty() //No Label
-		,true //bNoIndent
-		,false)); //bSearchable
+			// If no viewport has yet been set, we have nothing more to do.
+			if (!CurrentViewport)
+			{
+				return;
+			}
+
+			// Make sure we know about that viewport
+			TrackNewViewportIfRequired(CurrentViewport);
+
+			// Fetch configuration for this viewport. If none were made, we'll populate UI with default values
+			const FOpenColorIODisplayConfiguration& Configuration =
+				IOpenColorIOModule::Get().GetDisplayManager().FindOrAddDisplayConfiguration(CurrentViewport->GetClient());
+
+			// Add OCIO display section
+			FToolMenuSection& Section = DynamicSectionMenu->AddSection(
+				"DisplayConfiguration", LOCTEXT("DisplayConfiguration_Label", "Display Configuration"));
+			Section.AddEntry(FToolMenuEntry::InitWidget("DisplayConfigurationWidget",
+				SNew(SOpenColorIODisplay)
+					.Viewport(CurrentViewport)
+					.InitialConfiguration(Configuration)
+					.OnConfigurationChanged(FOnDisplayConfigurationChanged::CreateRaw(
+						this, &FOpenColorIOEditorModule::OnDisplayConfigurationChanged)),
+				FText::GetEmpty(), // No Label
+				true, // bNoIndent
+				false)); // bSearchable
+		}));
 }
 
 void FOpenColorIOEditorModule::OnDisplayConfigurationChanged(const FOpenColorIODisplayConfiguration& NewConfiguration)
