OCIO: Fix CIS warning about missing associated viewport client.

The assumption made years ago previously caused issues with lumen viewports, and now in automation tests. Removing the ensure yields no further issues.

#rnx

diff --git a/Engine/Plugins/Compositing/OpenColorIO/Source/OpenColorIOEditor/Private/OpenColorIOEditorModule.cpp b/Engine/Plugins/Compositing/OpenColorIO/Source/OpenColorIOEditor/Private/OpenColorIOEditorModule.cpp
--- a/Engine/Plugins/Compositing/OpenColorIO/Source/OpenColorIOEditor/Private/OpenColorIOEditorModule.cpp
+++ b/Engine/Plugins/Compositing/OpenColorIO/Source/OpenColorIOEditor/Private/OpenColorIOEditorModule.cpp
@@ -285,8 +285,7 @@ void FOpenColorIOEditorModule::TrackNewViewportIfRequired(FViewport* Viewport)
 		const TArray<FLevelEditorViewportClient*> LevelViewportClients = GEditor->GetLevelViewportClients();
 		FLevelEditorViewportClient* const* AssociatedClient = LevelViewportClients.FindByPredicate([Viewport](const FLevelEditorViewportClient* Other) { return Other == Viewport->GetClient(); });
 
-		// Active viewport should always have a client
-		if (ensure(AssociatedClient))
+		if (AssociatedClient)
 		{
 			FLevelEditorViewportClient* Client = *AssociatedClient;
 			TSharedPtr<SLevelViewport> LevelViewport = StaticCastSharedPtr<SLevelViewport>(Client->GetEditorViewportWidget());
