OCIO: Allow negative colors through the generated shader.

diff --git a/Engine/Plugins/Compositing/OpenColorIO/Shaders/Private/OpenColorIOShader.usf b/Engine/Plugins/Compositing/OpenColorIO/Shaders/Private/OpenColorIOShader.usf
--- a/Engine/Plugins/Compositing/OpenColorIO/Shaders/Private/OpenColorIOShader.usf
+++ b/Engine/Plugins/Compositing/OpenColorIO/Shaders/Private/OpenColorIOShader.usf
@@ -11,7 +11,7 @@
 Texture2D		InputTexture;
 SamplerState	InputTextureSampler;
 float			Gamma;
-uint TransformAlpha;
+uint			TransformAlpha;
 
 /*=============================================================================
  * OCIO generated start
@@ -47,6 +47,11 @@ float4 Premultiply(float4 InColor)
 	return float4(InColor.xyz * InColor.w, InColor.w);
 }
 
+// Note: For our (zero-tolerance) tests to succeed against a no-op, precise is required with the pow function.
+float3 DisplayGammaCorrection(precise float3 Color, precise float Gamma)
+{
+	return sign(Color) * pow(abs(Color), Gamma);
+}
 
 /* Pixel shader
 =============================================================================*/
@@ -55,8 +60,7 @@ void MainPS(in noperspective float4 InUVAndScreenPos : TEXCOORD0,
 	out float4 OutColor : SV_Target0)
 {
 	float4 InputColor = InputTexture.Sample(InputTextureSampler, InUVAndScreenPos.xy);
-	InputColor.xyz = max(0.0, InputColor.xyz);
-	InputColor.xyz = pow(InputColor.xyz, Gamma);
+	InputColor.xyz = DisplayGammaCorrection(InputColor.xyz, Gamma);
 
 	if(TransformAlpha > 0)
 	{
