OCIO: Renamed new shader function parameters for consitency.

#rnx

diff --git a/Engine/Plugins/Compositing/OpenColorIO/Shaders/Private/OpenColorIOShader.usf b/Engine/Plugins/Compositing/OpenColorIO/Shaders/Private/OpenColorIOShader.usf
--- a/Engine/Plugins/Compositing/OpenColorIO/Shaders/Private/OpenColorIOShader.usf
+++ b/Engine/Plugins/Compositing/OpenColorIO/Shaders/Private/OpenColorIOShader.usf
@@ -48,9 +48,9 @@ float4 Premultiply(float4 InColor)
 }
 
 // Note: For our (zero-tolerance) tests to succeed against a no-op, precise is required with the pow function.
-float3 DisplayGammaCorrection(precise float3 Color, precise float Gamma)
+float3 DisplayGammaCorrection(precise float3 InColor, precise float InGamma)
 {
-	return sign(Color) * pow(abs(Color), Gamma);
+	return sign(InColor) * pow(abs(InColor), InGamma);
 }
 
 /* Pixel shader
