OCIO: Robust fix for use after-free, and addresses incorrect previous change which prevented display-view selection.

diff --git a/Engine/Plugins/Compositing/OpenColorIO/Source/OpenColorIOEditor/Private/OpenColorIOColorSpaceCustomization.cpp b/Engine/Plugins/Compositing/OpenColorIO/Source/OpenColorIOEditor/Private/OpenColorIOColorSpaceCustomization.cpp
--- a/Engine/Plugins/Compositing/OpenColorIO/Source/OpenColorIOEditor/Private/OpenColorIOColorSpaceCustomization.cpp
+++ b/Engine/Plugins/Compositing/OpenColorIO/Source/OpenColorIOEditor/Private/OpenColorIOColorSpaceCustomization.cpp
@@ -19,16 +19,6 @@ void FOpenColorIOColorSpaceCustomization::CustomizeHeader(TSharedRef<IPropertyHa
 
 	if (CachedProperty->GetNumPerObjectValues() == 1 && CachedProperty->IsValidHandle())
 	{
-		FProperty* Property = CachedProperty->GetProperty();
-		check(Property && CastField<FStructProperty>(Property) && CastField<FStructProperty>(Property)->Struct && CastField<FStructProperty>(Property)->Struct->IsChildOf(FOpenColorIOColorSpace::StaticStruct()));
-
-		TArray<void*> RawData;
-		CachedProperty->AccessRawData(RawData);
-
-		check(RawData.Num() == 1);
-		FOpenColorIOColorSpace* ColorSpaceValue = reinterpret_cast<FOpenColorIOColorSpace*>(RawData[0]);
-
-		check(ColorSpaceValue);
 		TSharedPtr<IPropertyUtilities> PropertyUtils = CustomizationUtils.GetPropertyUtilities();
 
 		HeaderRow
@@ -45,14 +35,25 @@ void FOpenColorIOColorSpaceCustomization::CustomizeHeader(TSharedRef<IPropertyHa
 				.VAlign(VAlign_Center)
 				[
 					SNew(STextBlock)
-					.Text(MakeAttributeLambda([ColorSpaceValue] {
-							const FString ColorSpaceName = ColorSpaceValue->ToString();
-							
-							if(!ColorSpaceName.IsEmpty())
+					.Text(MakeAttributeLambda([WeakPropertyHandle = TWeakPtr<IPropertyHandle>(InPropertyHandle)]
+						{
+							if (TSharedPtr<IPropertyHandle> PropertyHandle = WeakPropertyHandle.Pin())
 							{
-								return FText::FromString(ColorSpaceName);
-							}
+								TArray<void*> RawData;
+								PropertyHandle->AccessRawData(RawData);
+								check(RawData.Num() == 1);
+
+								const FOpenColorIOColorSpace* ColorSpaceValue = reinterpret_cast<FOpenColorIOColorSpace*>(RawData[0]);
+								check(ColorSpaceValue);
 
+								const FString ColorSpaceName = ColorSpaceValue->ToString();
+
+								if (!ColorSpaceName.IsEmpty())
+								{
+									return FText::FromString(ColorSpaceName);
+								}
+							}
+							
 							return LOCTEXT("None", "<None>");
 						}))
 				]
@@ -76,18 +77,6 @@ void FOpenColorIODisplayViewCustomization::CustomizeHeader(TSharedRef<IPropertyH
 
 	if (CachedProperty->GetNumPerObjectValues() == 1 && CachedProperty->IsValidHandle())
 	{
-		FProperty* Property = CachedProperty->GetProperty();
-		check(Property && CastField<FStructProperty>(Property) && CastField<FStructProperty>(Property)->Struct && CastField<FStructProperty>(Property)->Struct->IsChildOf(FOpenColorIODisplayView::StaticStruct()));
-
-		TArray<void*> RawData;
-		CachedProperty->AccessRawData(RawData);
-		check(RawData.Num() == 1);
-		
-		const FOpenColorIODisplayView* DisplayViewValue = reinterpret_cast<const FOpenColorIODisplayView*>(RawData[0]);
-		check(DisplayViewValue);
-
-		const FString DisplayViewName = DisplayViewValue->ToString();
-		
 		TSharedPtr<IPropertyUtilities> PropertyUtils = CustomizationUtils.GetPropertyUtilities();
 
 		HeaderRow
@@ -104,11 +93,26 @@ void FOpenColorIODisplayViewCustomization::CustomizeHeader(TSharedRef<IPropertyH
 				.VAlign(VAlign_Center)
 				[
 					SNew(STextBlock)
-					.Text(MakeAttributeLambda([DisplayViewName]
+					.Text(MakeAttributeLambda([WeakPropertyHandle = TWeakPtr<IPropertyHandle>(InPropertyHandle)]
 						{
-							if(!DisplayViewName.IsEmpty())
+							if (TSharedPtr<IPropertyHandle> PropertyHandle = WeakPropertyHandle.Pin())
 							{
-								return FText::FromString(DisplayViewName);
+								FProperty* Property = PropertyHandle->GetProperty();
+								check(Property && CastField<FStructProperty>(Property) && CastField<FStructProperty>(Property)->Struct && CastField<FStructProperty>(Property)->Struct->IsChildOf(FOpenColorIODisplayView::StaticStruct()));
+
+								TArray<void*> RawData;
+								PropertyHandle->AccessRawData(RawData);
+								check(RawData.Num() == 1);
+
+								const FOpenColorIODisplayView* DisplayViewValue = reinterpret_cast<const FOpenColorIODisplayView*>(RawData[0]);
+								check(DisplayViewValue);
+
+								const FString DisplayViewName = DisplayViewValue->ToString();
+
+								if (!DisplayViewName.IsEmpty())
+								{
+									return FText::FromString(DisplayViewName);
+								}
 							}
 
 							return LOCTEXT("None", "<None>");
