OpenColorIO: Deprecated function cleanup.

#rnx

diff --git a/Engine/Plugins/Compositing/OpenColorIO/Source/OpenColorIO/Private/OpenColorIOColorTransform.cpp b/Engine/Plugins/Compositing/OpenColorIO/Source/OpenColorIO/Private/OpenColorIOColorTransform.cpp
--- a/Engine/Plugins/Compositing/OpenColorIO/Source/OpenColorIO/Private/OpenColorIOColorTransform.cpp
+++ b/Engine/Plugins/Compositing/OpenColorIO/Source/OpenColorIO/Private/OpenColorIOColorTransform.cpp
@@ -139,16 +139,6 @@ UOpenColorIOColorTransform::UOpenColorIOColorTransform(const FObjectInitializer&
 {
 }
 
-bool UOpenColorIOColorTransform::Initialize(const FString& InSourceColorSpace, const FString& InDestinationColorSpace, const TMap<FString, FString>& /*InContextKeyValues*/)
-{
-	return Initialize(InSourceColorSpace, InDestinationColorSpace);
-}
-
-bool UOpenColorIOColorTransform::Initialize(const FString& InSourceColorSpace, const FString& InDisplay, const FString& InView, EOpenColorIOViewTransformDirection InDirection, const TMap<FString, FString>& /*InContextKeyValues*/)
-{
-	return Initialize(InSourceColorSpace, InDisplay, InView, InDirection);
-}
-
 void UOpenColorIOColorTransform::Serialize(FArchive& Ar)
 {
 	Super::Serialize(Ar);
diff --git a/Engine/Plugins/Compositing/OpenColorIO/Source/OpenColorIO/Private/OpenColorIOConfiguration.cpp b/Engine/Plugins/Compositing/OpenColorIO/Source/OpenColorIO/Private/OpenColorIOConfiguration.cpp
--- a/Engine/Plugins/Compositing/OpenColorIO/Source/OpenColorIO/Private/OpenColorIOConfiguration.cpp
+++ b/Engine/Plugins/Compositing/OpenColorIO/Source/OpenColorIO/Private/OpenColorIOConfiguration.cpp
@@ -535,13 +535,6 @@ void UOpenColorIOConfiguration::PostLoad()
 	}
 }
 
-void UOpenColorIOConfiguration::GetAssetRegistryTags(TArray<FAssetRegistryTag>& OutTags) const
-{
-	PRAGMA_DISABLE_DEPRECATION_WARNINGS;
-	Super::GetAssetRegistryTags(OutTags);
-	PRAGMA_ENABLE_DEPRECATION_WARNINGS;
-}
-
 void UOpenColorIOConfiguration::GetAssetRegistryTags(FAssetRegistryTagsContext RegistryTagsContext) const
 {
 	Super::GetAssetRegistryTags(RegistryTagsContext);
diff --git a/Engine/Plugins/Compositing/OpenColorIO/Source/OpenColorIO/Private/OpenColorIOShared.cpp b/Engine/Plugins/Compositing/OpenColorIO/Source/OpenColorIO/Private/OpenColorIOShared.cpp
--- a/Engine/Plugins/Compositing/OpenColorIO/Source/OpenColorIO/Private/OpenColorIOShared.cpp
+++ b/Engine/Plugins/Compositing/OpenColorIO/Source/OpenColorIO/Private/OpenColorIOShared.cpp
@@ -209,11 +209,6 @@ void FOpenColorIOTransformResource::SerializeShaderMap(FArchive& Ar)
 	}
 }
 
-void FOpenColorIOTransformResource::SetupResource(ERHIFeatureLevel::Type InFeatureLevel, const FString& InShaderCodeHash, const FString& InShadercode, const FString& InRawConfigHash, const FString& InFriendlyName, const FName& InAssetPath, EOpenColorIOWorkingColorSpaceTransform /*InWorkingColorSpaceTransformType*/)
-{
-	SetupResource(InFeatureLevel, InShadercode, InShadercode, InRawConfigHash, InFriendlyName, InAssetPath);
-}
-
 void FOpenColorIOTransformResource::SetupResource(ERHIFeatureLevel::Type InFeatureLevel, const FString& InShaderCodeHash, const FString& InShadercode, const FString& InRawConfigHash, const FString& InFriendlyName, const FName& InAssetPath)
 {
 	// When this happens we assume that shader was cooked and we don't need to do anything. 
diff --git a/Engine/Plugins/Compositing/OpenColorIO/Source/OpenColorIO/Public/OpenColorIOColorTransform.h b/Engine/Plugins/Compositing/OpenColorIO/Source/OpenColorIO/Public/OpenColorIOColorTransform.h
--- a/Engine/Plugins/Compositing/OpenColorIO/Source/OpenColorIO/Public/OpenColorIOColorTransform.h
+++ b/Engine/Plugins/Compositing/OpenColorIO/Source/OpenColorIO/Public/OpenColorIOColorTransform.h
@@ -34,10 +34,6 @@ public:
 	virtual ~UOpenColorIOColorTransform() {};
 
 public:
-	UE_DEPRECATED(5.4, "This method is deprecated, please use Initialize without the context argument.")
-	bool Initialize(const FString& InSourceColorSpace, const FString& InDestinationColorSpace, const TMap<FString, FString>& InContextKeyValues);
-	UE_DEPRECATED(5.4, "This method is deprecated, please use Initialize without the context argument.")
-	bool Initialize(const FString& InSourceColorSpace, const FString& InDisplay, const FString& InView, EOpenColorIOViewTransformDirection InDirection, const TMap<FString, FString>& InContextKeyValues);
 
 	/** Initialize resources for color space transform. */
 	bool Initialize(const FString& InSourceColorSpace, const FString& InDestinationColorSpace);
@@ -45,25 +41,6 @@ public:
 	/** Initialize resources for display-view transform. */
 	bool Initialize(const FString& InSourceColorSpace, const FString& InDisplay, const FString& InView, EOpenColorIOViewTransformDirection InDirection);
 
-	/**
-	 * Cache resource shaders for rendering.
-	 * If a matching shader map is not found in memory or the DDC, a new one will be compiled.
-	 */
-	UE_DEPRECATED_FORGAME(5.4, "Do not use. Will be made private in 5.5")
-	void CacheResourceShadersForRendering(bool bRegenerateId = false);
-	
-	UE_DEPRECATED_FORGAME(5.4, "Do not use. Will be made private in 5.5")
-	void CacheShadersForResources(EShaderPlatform InShaderPlatform, FOpenColorIOTransformResource* InResourcesToCache, bool bApplyCompletedShaderMapForRendering, bool bIsCooking, const ITargetPlatform* TargetPlatform = nullptr);
-
-	/**
-	 * Returns the desired resources required to apply this transform during rendering.
-	 */
-	UE_DEPRECATED(5.4, "This method is deprecated.")
-	bool GetRenderResources(ERHIFeatureLevel::Type InFeatureLevel, FOpenColorIOTransformResource*& OutShaderResource, TSortedMap<int32, FTextureResource*>& OutTextureResources) const
-	{
-		return false;
-	}
-
 	/**
 	 * Returns the desired resources required to apply this transform during rendering.
 	 */
@@ -115,30 +92,34 @@ public:
 	// Get the owner config's context key-values.
 	TMap<FString, FString> GetContextKeyValues() const;
 
-protected:
+private:
+
+	/**
+	 * Cache resource shaders for rendering.
+	 * If a matching shader map is not found in memory or the DDC, a new one will be compiled.
+	 */
+	void CacheResourceShadersForRendering(bool bRegenerateId = false);
+
+	void CacheShadersForResources(EShaderPlatform InShaderPlatform, FOpenColorIOTransformResource* InResourcesToCache, bool bApplyCompletedShaderMapForRendering, bool bIsCooking, const ITargetPlatform* TargetPlatform = nullptr);
 
 	/**
 	 * Helper function to serialize shader maps for the given color transform resources.
 	 */
-	UE_DEPRECATED_FORGAME(5.4, "Do not use. Will be made private in 5.5")
 	static void SerializeOpenColorIOShaderMaps(const TMap<const ITargetPlatform*, TArray<FOpenColorIOTransformResource*>>* PlatformColorTransformResourcesToSavePtr, FArchive& Ar, TArray<FOpenColorIOTransformResource>&  OutLoadedResources);
 	
 	/**
 	 * Helper function to register serialized shader maps for the given color transform resources
 	 */
-	UE_DEPRECATED_FORGAME(5.4, "Do not use. Will be made private in 5.5")
 	static void ProcessSerializedShaderMaps(UOpenColorIOColorTransform* Owner, TArray<FOpenColorIOTransformResource>& LoadedResources, FOpenColorIOTransformResource* (&OutColorTransformResourcesLoaded)[ERHIFeatureLevel::Num]);
 	
 	/**
 	 * Returns a Guid for the LUT based on its unique identifier, name and the OCIO DDC key.
 	 */
-	UE_DEPRECATED_FORGAME(5.4, "Do not use. Will be made private in 5.5")
 	static void GetOpenColorIOLUTKeyGuid(const FString& InProcessorIdentifier, const FName& InName, FGuid& OutLutGuid );
 
 	/**
 	 * Helper function returning the color space transform name based on source and destination color spaces.
 	 */
-	UE_DEPRECATED_FORGAME(5.4, "Do not use. Will be made private in 5.5")
 	FString GetTransformFriendlyName() const;
 
 #if WITH_EDITOR
@@ -146,24 +127,18 @@ protected:
 	 * Fetch shader code and hash from the OCIO library
 	 * @return: true if shader could be generated from the library
 	 */
-	UE_DEPRECATED_FORGAME(5.4, "Do not use. Will be made private in 5.5")
 	bool UpdateShaderInfo(FString& OutShaderCodeHash, FString& OutShaderCode, FString& OutRawConfigHash);
 
 	/**
 	 * Helper function taking raw 3D LUT data coming from the library and initializing a UVolumeTexture with it.
 	 */
-	UE_DEPRECATED_FORGAME(5.4, "Do not use. Will be made private in 5.5")
 	TObjectPtr<UTexture> CreateTexture3DLUT(const FString& InProcessorIdentifier, const FName& InName, uint32 InLutLength, TextureFilter InFilter, const float* InSourceData);
 
 	/**
 	 * Helper function taking raw 1D LUT data coming from the library and initializing a UTexture with it.
 	 */
-	UE_DEPRECATED_FORGAME(5.4, "Do not use. Will be made private in 5.5")
 	TObjectPtr<UTexture> CreateTexture1DLUT(const FString& InProcessorIdentifier, const FName& InName, uint32 InTextureWidth, uint32 InTextureHeight, TextureFilter InFilter, bool bRedChannelOnly, const float* InSourceData );
-#endif //WITH_EDITOR
 
-private:
-#if WITH_EDITOR
 	/**
 	 * Create the transform processor(s) and generate its resources. Editor-only, in game mode GPU resources are already present.
 	 */
diff --git a/Engine/Plugins/Compositing/OpenColorIO/Source/OpenColorIO/Public/OpenColorIOConfiguration.h b/Engine/Plugins/Compositing/OpenColorIO/Source/OpenColorIO/Public/OpenColorIOConfiguration.h
--- a/Engine/Plugins/Compositing/OpenColorIO/Source/OpenColorIO/Public/OpenColorIOConfiguration.h
+++ b/Engine/Plugins/Compositing/OpenColorIO/Source/OpenColorIO/Public/OpenColorIOConfiguration.h
@@ -39,12 +39,6 @@ public:
 public:
 
 	bool IsTransformReady(const FOpenColorIOColorConversionSettings& InSettings);
-	
-	UE_DEPRECATED(5.4, "This method is deprecated.")
-	bool GetRenderResources(ERHIFeatureLevel::Type InFeatureLevel, const FOpenColorIOColorConversionSettings& InSettings, FOpenColorIOTransformResource*& OutShaderResource, TSortedMap<int32, FTextureResource*>& OutTextureResources)
-	{
-		return false;
-	}
 
 	bool GetRenderResources(ERHIFeatureLevel::Type InFeatureLevel, const FOpenColorIOColorConversionSettings& InSettings, FOpenColorIOTransformResource*& OutShaderResource, TSortedMap<int32, TWeakObjectPtr<UTexture>>& OutTextureResources);
 	
@@ -106,8 +100,6 @@ public:
 	virtual void PostInitProperties() override;
 	virtual void PostLoad() override;
 	virtual void GetAssetRegistryTags(FAssetRegistryTagsContext RegistryTagsContext) const override;
-	UE_DEPRECATED(5.4, "Implement the version that takes FAssetRegistryTagsContext instead.")
-	virtual void GetAssetRegistryTags(TArray<FAssetRegistryTag>& OutTags) const override;
 #if WITH_EDITORONLY_DATA
 	static void DeclareConstructClasses(TArray<FTopLevelAssetPath>& OutConstructClasses, const UClass* SpecificSubclass);
 #endif
diff --git a/Engine/Plugins/Compositing/OpenColorIO/Source/OpenColorIO/Public/OpenColorIOShared.h b/Engine/Plugins/Compositing/OpenColorIO/Source/OpenColorIO/Public/OpenColorIOShared.h
--- a/Engine/Plugins/Compositing/OpenColorIO/Source/OpenColorIO/Public/OpenColorIOShared.h
+++ b/Engine/Plugins/Compositing/OpenColorIO/Source/OpenColorIO/Public/OpenColorIOShared.h
@@ -455,17 +455,6 @@ public:
 
 	const FString& GetFriendlyName()	const { return FriendlyName; }
 
-	UE_DEPRECATED(5.4, "SetupResource with EOpenColorIOWorkingColorSpaceTransform deprecated.")
-	void SetupResource(
-		ERHIFeatureLevel::Type InFeatureLevel,
-		const FString& InShaderCodeHash,
-		const FString& InShadercode,
-		const FString& InRawConfigHash,
-		const FString& InFriendlyName,
-		const FName& InAssetPath,
-		EOpenColorIOWorkingColorSpaceTransform InWorkingColorSpaceTransformType
-	);
-
 	void SetupResource(
 		ERHIFeatureLevel::Type InFeatureLevel,
 		const FString& InShaderCodeHash,
@@ -508,8 +497,6 @@ public:
 	
 	bool IsSame(const FOpenColorIOShaderMapId& InId) const;
 
-	UE_DEPRECATED(5.4, "GetWorkingColorSpaceTransformType is deprecated.")
-	EOpenColorIOWorkingColorSpaceTransform GetWorkingColorSpaceTransformType() const { return EOpenColorIOWorkingColorSpaceTransform::None; }
 protected:
 #if WITH_EDITOR
 	/**
diff --git a/Engine/Plugins/Runtime/nDisplay/Source/DisplayClusterConfiguration/Private/DisplayClusterConfigurationTypes_OCIO.cpp b/Engine/Plugins/Runtime/nDisplay/Source/DisplayClusterConfiguration/Private/DisplayClusterConfigurationTypes_OCIO.cpp
--- a/Engine/Plugins/Runtime/nDisplay/Source/DisplayClusterConfiguration/Private/DisplayClusterConfigurationTypes_OCIO.cpp
+++ b/Engine/Plugins/Runtime/nDisplay/Source/DisplayClusterConfiguration/Private/DisplayClusterConfigurationTypes_OCIO.cpp
@@ -29,11 +29,6 @@ void FDisplayClusterConfigurationOCIOConfiguration::PostSerialize(const FArchive
 	}
 }
 
-bool FDisplayClusterConfigurationOCIOConfiguration::IsEnabled() const
-{
-	return bIsEnabled && ColorConfiguration.IsValid();
-}
-
 ///////////////////////////////////////////////////////////////////////////////////////
 // FDisplayClusterConfigurationOCIOProfile
 ///////////////////////////////////////////////////////////////////////////////////////
@@ -60,11 +55,6 @@ void FDisplayClusterConfigurationOCIOProfile::PostSerialize(const FArchive& Ar)
 	}
 }
 
-bool FDisplayClusterConfigurationOCIOProfile::IsEnabled() const
-{
-	return bIsEnabled && ColorConfiguration.IsValid();
-}
-
 bool FDisplayClusterConfigurationOCIOProfile::IsEnabledForObject(const FString& InObjectId) const
 {
 	if (bIsEnabled)
diff --git a/Engine/Plugins/Runtime/nDisplay/Source/DisplayClusterConfiguration/Public/DisplayClusterConfigurationTypes_OCIO.h b/Engine/Plugins/Runtime/nDisplay/Source/DisplayClusterConfiguration/Public/DisplayClusterConfigurationTypes_OCIO.h
--- a/Engine/Plugins/Runtime/nDisplay/Source/DisplayClusterConfiguration/Public/DisplayClusterConfigurationTypes_OCIO.h
+++ b/Engine/Plugins/Runtime/nDisplay/Source/DisplayClusterConfiguration/Public/DisplayClusterConfigurationTypes_OCIO.h
@@ -33,10 +33,6 @@ public:
 
 	bool Serialize(FArchive& Ar);
 	void PostSerialize(const FArchive& Ar);
-	
-	UE_DEPRECATED(5.3, "This method is deprecated.")
-	/** Return true if configuration valid */
-	bool IsEnabled() const;
 };
 
 template<> struct TStructOpsTypeTraits<FDisplayClusterConfigurationOCIOConfiguration> : public TStructOpsTypeTraitsBase2<FDisplayClusterConfigurationOCIOConfiguration>
@@ -81,10 +77,6 @@ public:
 
 	/** Return true if the configuration is valid for the input object */
 	bool IsEnabledForObject(const FString& InObjectId) const;
-	
-	UE_DEPRECATED(5.3, "This method is deprecated.")
-	/** Return true if configuration valid */
-	bool IsEnabled() const;
 };
 
 template<> struct TStructOpsTypeTraits<FDisplayClusterConfigurationOCIOProfile> : public TStructOpsTypeTraitsBase2<FDisplayClusterConfigurationOCIOProfile>
