OCIO: Fixed slate crash on config string deletion.

We null-check against existing color spaces or display-views. They are still displayed, but have already been deleted by config reloading logic.

#rnx

diff --git a/Engine/Plugins/Compositing/OpenColorIO/Source/OpenColorIOEditor/Private/OpenColorIOColorSpaceCustomization.cpp b/Engine/Plugins/Compositing/OpenColorIO/Source/OpenColorIOEditor/Private/OpenColorIOColorSpaceCustomization.cpp
--- a/Engine/Plugins/Compositing/OpenColorIO/Source/OpenColorIOEditor/Private/OpenColorIOColorSpaceCustomization.cpp
+++ b/Engine/Plugins/Compositing/OpenColorIO/Source/OpenColorIOEditor/Private/OpenColorIOColorSpaceCustomization.cpp
@@ -37,20 +37,22 @@ void FOpenColorIOColorSpaceCustomization::CustomizeHeader(TSharedRef<IPropertyHa
 					SNew(STextBlock)
 					.Text(MakeAttributeLambda([WeakPropertyHandle = TWeakPtr<IPropertyHandle>(InPropertyHandle)]
 						{
-							if (TSharedPtr<IPropertyHandle> PropertyHandle = WeakPropertyHandle.Pin())
+							TSharedPtr<IPropertyHandle> PropertyHandle = WeakPropertyHandle.Pin();
+							if (PropertyHandle.IsValid() && PropertyHandle->IsValidHandle())
 							{
 								TArray<void*> RawData;
 								PropertyHandle->AccessRawData(RawData);
 								check(RawData.Num() == 1);
 
 								const FOpenColorIOColorSpace* ColorSpaceValue = reinterpret_cast<FOpenColorIOColorSpace*>(RawData[0]);
-								check(ColorSpaceValue);
-
-								const FString ColorSpaceName = ColorSpaceValue->ToString();
-
-								if (!ColorSpaceName.IsEmpty())
+								if (ColorSpaceValue != nullptr)
 								{
-									return FText::FromString(ColorSpaceName);
+									const FString ColorSpaceName = ColorSpaceValue->ToString();
+
+									if (!ColorSpaceName.IsEmpty())
+									{
+										return FText::FromString(ColorSpaceName);
+									}
 								}
 							}
 							
@@ -95,7 +97,8 @@ void FOpenColorIODisplayViewCustomization::CustomizeHeader(TSharedRef<IPropertyH
 					SNew(STextBlock)
 					.Text(MakeAttributeLambda([WeakPropertyHandle = TWeakPtr<IPropertyHandle>(InPropertyHandle)]
 						{
-							if (TSharedPtr<IPropertyHandle> PropertyHandle = WeakPropertyHandle.Pin())
+							TSharedPtr<IPropertyHandle> PropertyHandle = WeakPropertyHandle.Pin();
+							if (PropertyHandle.IsValid() && PropertyHandle->IsValidHandle())
 							{
 								FProperty* Property = PropertyHandle->GetProperty();
 								check(Property && CastField<FStructProperty>(Property) && CastField<FStructProperty>(Property)->Struct && CastField<FStructProperty>(Property)->Struct->IsChildOf(FOpenColorIODisplayView::StaticStruct()));
@@ -105,13 +108,14 @@ void FOpenColorIODisplayViewCustomization::CustomizeHeader(TSharedRef<IPropertyH
 								check(RawData.Num() == 1);
 
 								const FOpenColorIODisplayView* DisplayViewValue = reinterpret_cast<const FOpenColorIODisplayView*>(RawData[0]);
-								check(DisplayViewValue);
-
-								const FString DisplayViewName = DisplayViewValue->ToString();
-
-								if (!DisplayViewName.IsEmpty())
+								if (DisplayViewValue != nullptr)
 								{
-									return FText::FromString(DisplayViewName);
+									const FString DisplayViewName = DisplayViewValue->ToString();
+
+									if (!DisplayViewName.IsEmpty())
+									{
+										return FText::FromString(DisplayViewName);
+									}
 								}
 							}
 
