OCIO: Shader parameter refactor to support dynamic texture bindings.

This removes the static list of 1D/3D lookup texture parameters in favor of a dynamic one specifically built to match the generated transform code. This follows patterns already in use by ComputeGraph & Niagara, extended for pixel shaders: see FShaderParametersMetadataBuilder -> FShaderParametersMetadata -> FShaderParameterBindings.

This refactor was necessary for the upcoming ACES 2.0 built-in support, and may also be used to expose dynamic properties.

diff --git a/Engine/Plugins/Compositing/OpenColorIO/Source/OpenColorIO/Private/OpenColorIOColorTransform.cpp b/Engine/Plugins/Compositing/OpenColorIO/Source/OpenColorIO/Private/OpenColorIOColorTransform.cpp
--- a/Engine/Plugins/Compositing/OpenColorIO/Source/OpenColorIO/Private/OpenColorIOColorTransform.cpp
+++ b/Engine/Plugins/Compositing/OpenColorIO/Source/OpenColorIO/Private/OpenColorIOColorTransform.cpp
@@ -12,6 +12,7 @@
 #include "UObject/UE5MainStreamObjectVersion.h"
 #include "UObject/UObjectIterator.h"
 #include "DataDrivenShaderPlatformInfo.h"
+#include "ShaderParameterMetadataBuilder.h"
 
 #include UE_INLINE_GENERATED_CPP_BY_NAME(OpenColorIOColorTransform)
 
@@ -84,9 +85,8 @@ void UOpenColorIOColorTransform::SerializeOpenColorIOShaderMaps(const TMap<const
 
 		for (int32 ResourceIndex = 0; ResourceIndex < NumLoadedResources; ResourceIndex++)
 		{
-			FOpenColorIOTransformResource LoadedResource;
+			FOpenColorIOTransformResource& LoadedResource = OutLoadedResources.AddDefaulted_GetRef();
 			LoadedResource.SerializeShaderMap(Ar);
-			OutLoadedResources.Add(LoadedResource);
 		}
 	}
 }
@@ -217,7 +217,7 @@ void UOpenColorIOColorTransform::CacheResourceShadersForCooking(EShaderPlatform
 	FOpenColorIOTransformResource* NewResource = new FOpenColorIOTransformResource();
 
 	FName AssetPath = GetOutermost()->GetFName();
-	NewResource->SetupResource((ERHIFeatureLevel::Type)TargetFeatureLevel, InShaderHash, InShaderCode, InRawConfigHash, GetTransformFriendlyName(), AssetPath);
+	NewResource->SetupResource((ERHIFeatureLevel::Type)TargetFeatureLevel, InShaderHash, InShaderCode, InRawConfigHash, GetTransformFriendlyName(), AssetPath, BuildShaderParamMetadata());
 
 	const bool bApplyCompletedShaderMap = false;
 	const bool bIsCooking = true;
@@ -237,7 +237,7 @@ void UOpenColorIOColorTransform::ProcessTransformForGPU()
 		// The resources that were returned earlier using GetRenderResources() and GetRenderPassResources() will also be deleted.
 		// Other plugins store OCIO resources in the FOpenColorIORenderPassResources structure, which uses pointers to FTextureResources
 		// that are used without reference counting. Once deleted in the game thread, they become invalid in the rendering thread.
-		Textures.Reset();
+		LookupTextures.Reset();
 
 		const UOpenColorIOSettings* Settings = GetDefault<UOpenColorIOSettings>();
 		const FOpenColorIOWrapperGPUProcessor GPUProcessor = FOpenColorIOWrapperGPUProcessor(TransformProcessor, Settings->bUseLegacyProcessor);
@@ -252,12 +252,9 @@ void UOpenColorIOColorTransform::ProcessTransformForGPU()
 			const float* TextureValues = 0x0;
 
 			bool bSuccess = GPUProcessor.Get3DTexture(Index, TextureName, EdgeLength, Filter, TextureValues);
-			check(bSuccess);
+			checkf(bSuccess, TEXT("Failed to read OCIO 3D LUT data."));
 
-			TObjectPtr<UTexture> Result = CreateTexture3DLUT(GpuProcessorHash, TextureName, EdgeLength, Filter, TextureValues);
-
-			const int32 SlotIndex = TextureName.GetNumber() - 1; // Rely on FName's index number extraction for convenience
-			Textures.Add(SlotIndex, MoveTemp(Result));
+			LookupTextures.Add(CreateTexture3DLUT(GpuProcessorHash, TextureName, EdgeLength, Filter, TextureValues));
 		}
 
 		// Process 1D luts
@@ -273,13 +270,8 @@ void UOpenColorIOColorTransform::ProcessTransformForGPU()
 			bool bSuccess = GPUProcessor.GetTexture(Index, TextureName, TextureWidth, TextureHeight, Filter, bRedChannelOnly, TextureValues);
 			checkf(bSuccess, TEXT("Failed to read OCIO 1D LUT data."));
 
-			TObjectPtr<UTexture> Result = CreateTexture1DLUT(GpuProcessorHash, TextureName, TextureWidth, TextureHeight, Filter, bRedChannelOnly, TextureValues);
-
-			const int32 SlotIndex = TextureName.GetNumber() - 1; // Rely on FName's index number extraction for convenience
-			Textures.Add(SlotIndex, MoveTemp(Result));
+			LookupTextures.Add(CreateTexture1DLUT(GpuProcessorHash, TextureName, TextureWidth, TextureHeight, Filter, bRedChannelOnly, TextureValues));
 		}
-
-		ensureAlwaysMsgf(Textures.Num() <= (int32)OpenColorIOShader::MaximumTextureSlots, TEXT("Color transform %s exceeds our current limit of %u texture slots. Use the legacy processor instead."), *GetTransformFriendlyName(), OpenColorIOShader::MaximumTextureSlots);
 		
 		// Generate shader code
 		GPUProcessor.GetShader(GeneratedShaderHash, GeneratedShader);
@@ -440,7 +432,7 @@ void UOpenColorIOColorTransform::CacheResourceShadersForRendering(bool /*bRegene
 				RawConfigHash = ConfigurationOwner->GetConfigWrapper()->GetCacheID();
 				AssetPath = GetOutermost()->GetFName();
 #endif // WITH_EDITOR
-				TransformResource->SetupResource(CacheFeatureLevel, ShaderCodeHash, ShaderCode, RawConfigHash, GetTransformFriendlyName(), AssetPath);
+				TransformResource->SetupResource(CacheFeatureLevel, ShaderCodeHash, ShaderCode, RawConfigHash, GetTransformFriendlyName(), AssetPath, BuildShaderParamMetadata());
 
 				const bool bApplyCompletedShaderMap = true;
 
@@ -471,7 +463,7 @@ void UOpenColorIOColorTransform::CacheShadersForResources(EShaderPlatform InShad
 	}
 }
 
-bool UOpenColorIOColorTransform::GetRenderResources(ERHIFeatureLevel::Type InFeatureLevel, FOpenColorIOTransformResource*& OutShaderResource, TSortedMap<int32, TWeakObjectPtr<UTexture>>& OutTextureResources) const
+bool UOpenColorIOColorTransform::GetRenderResources(ERHIFeatureLevel::Type InFeatureLevel, FOpenColorIOTransformResource*& OutShaderResource, TArray<TWeakObjectPtr<UTexture>>& OutTextureResources) const
 {
 	if (!AreRenderResourcesReady())
 	{
@@ -482,17 +474,34 @@ bool UOpenColorIOColorTransform::GetRenderResources(ERHIFeatureLevel::Type InFea
 	
 	if (OutShaderResource)
 	{
-		OutTextureResources.Reserve(Textures.Num());
+		OutTextureResources.Reserve(LookupTextures.Num());
 
-		for (const TPair<int32, TObjectPtr<UTexture>>& Pair : Textures)
+		for (const TObjectPtr<UTexture>& Texture : LookupTextures)
 		{
-			OutTextureResources.Add(Pair.Key, Pair.Value);
+			OutTextureResources.Add(Texture);
 		}
 	}
 	
 	return OutShaderResource != nullptr;
 }
 
+// Deprecated
+bool UOpenColorIOColorTransform::GetRenderResources(ERHIFeatureLevel::Type InFeatureLevel, FOpenColorIOTransformResource*& OutShaderResource, TSortedMap<int32, TWeakObjectPtr<UTexture>>& OutTextureResources) const
+{
+	TArray<TWeakObjectPtr<UTexture>> OutTextureArray;
+	if (GetRenderResources(InFeatureLevel, OutShaderResource, OutTextureArray))
+	{
+		for (int32 Index = 0; Index < OutTextureArray.Num(); Index++)
+		{
+			OutTextureResources.Emplace(Index, OutTextureArray[Index]);
+		}
+
+		return true;
+	}
+
+	return false;
+}
+
 bool UOpenColorIOColorTransform::AreRenderResourcesReady() const
 {
 	check(IsInGameThread());
@@ -519,10 +528,8 @@ bool UOpenColorIOColorTransform::AreRenderResourcesReady() const
 	}
 
 	// Textures are optional, depending on the transform.
-	for (const TPair<int32, TObjectPtr<UTexture>>& TexturePair : Textures)
+	for (const TObjectPtr<UTexture>& Texture : LookupTextures)
 	{
-		const TObjectPtr<UTexture>& Texture = TexturePair.Value;
-
 		if (Texture->GetResource() == nullptr)
 		{
 			return false;
@@ -694,10 +701,52 @@ void UOpenColorIOColorTransform::FlushResourceShaderMaps()
 	}
 }
 
+TUniquePtr<FShaderParametersMetadata> UOpenColorIOColorTransform::BuildShaderParamMetadata()
+{
+	// We simply use static allocations for resource names, similar to ComputeFramework::AddParamForType.
+	static TArray<FString> AllTextureResourceNames;
+	static TArray<FString> AllSamplerResourceNames;
+
+	FShaderParametersMetadataBuilder Builder;
+
+	// First, we add our built-in shader paramaters
+	Builder.AddIncludedStruct<FOpenColorIOPixelShaderParameters>();
+
+	// Second, we add the dynamic texture & sampler parameters if present in the transform.
+	// We always add the texture resource followed by its sampler. The same order is at render-time,
+	// in FOpenColorIORendering::AddPass_RenderThread.
+	for (const TObjectPtr<UTexture>& Texture : LookupTextures)
+	{
+		if (IsValid(Texture))
+		{
+			int32 TexureIndex = AllTextureResourceNames.AddUnique(Texture->GetName());
+			int32 SamplerIndex = AllSamplerResourceNames.AddUnique(Texture->GetName() + TEXT("Sampler"));
+
+			Builder.AddTexture(*AllTextureResourceNames[TexureIndex], Texture.IsA(UVolumeTexture::StaticClass()) ? TEXT("TextureCube") : TEXT("Texture2D"));
+			Builder.AddSampler(*AllSamplerResourceNames[SamplerIndex], TEXT("SamplerState"));
+		}
+	}
+
+	return TUniquePtr<FShaderParametersMetadata>(Builder.Build(FShaderParametersMetadata::EUseCase::ShaderParameterStruct, TEXT("")));
+}
+
 void UOpenColorIOColorTransform::PostLoad()
 {
 	Super::PostLoad();
 
+#if WITH_EDITORONLY_DATA
+	PRAGMA_DISABLE_DEPRECATION_WARNINGS
+	if (!Textures_DEPRECATED.IsEmpty())
+	{
+		for (const TPair<int32, TObjectPtr<UTexture>>& TexturePair : Textures_DEPRECATED)
+		{
+			LookupTextures.Add(TexturePair.Value);
+		}
+		Textures_DEPRECATED.Empty();
+	}
+	PRAGMA_ENABLE_DEPRECATION_WARNINGS
+#endif // WITH_EDITORONLY_DATA
+
 	if (FApp::CanEverRender())
 	{
 		ProcessSerializedShaderMaps(this, LoadedTransformResources, ColorTransformResources);
diff --git a/Engine/Plugins/Compositing/OpenColorIO/Source/OpenColorIO/Private/OpenColorIOConfiguration.cpp b/Engine/Plugins/Compositing/OpenColorIO/Source/OpenColorIO/Private/OpenColorIOConfiguration.cpp
--- a/Engine/Plugins/Compositing/OpenColorIO/Source/OpenColorIO/Private/OpenColorIOConfiguration.cpp
+++ b/Engine/Plugins/Compositing/OpenColorIO/Source/OpenColorIO/Private/OpenColorIOConfiguration.cpp
@@ -72,12 +72,25 @@ bool UOpenColorIOConfiguration::IsTransformReady(const FOpenColorIOColorConversi
 	return false;
 }
 
+bool UOpenColorIOConfiguration::GetRenderResources(ERHIFeatureLevel::Type InFeatureLevel, const FOpenColorIOColorConversionSettings& InSettings, FOpenColorIOTransformResource*& OutShaderResource, TArray<TWeakObjectPtr<UTexture>>& OutTextureResources)
+{
+	TObjectPtr<const UOpenColorIOColorTransform> Transform = FindTransform(InSettings);
+	if (IsValid(Transform))
+	{
+		return Transform->GetRenderResources(InFeatureLevel, OutShaderResource, OutTextureResources);
+	}
+
+	return false;
+}
+
 bool UOpenColorIOConfiguration::GetRenderResources(ERHIFeatureLevel::Type InFeatureLevel, const FOpenColorIOColorConversionSettings& InSettings, FOpenColorIOTransformResource*& OutShaderResource, TSortedMap<int32, TWeakObjectPtr<UTexture>>& OutTextureResources)
 {
 	TObjectPtr<const UOpenColorIOColorTransform> Transform = FindTransform(InSettings);
 	if (IsValid(Transform))
 	{
+		PRAGMA_DISABLE_DEPRECATION_WARNINGS
 		return Transform->GetRenderResources(InFeatureLevel, OutShaderResource, OutTextureResources);
+		PRAGMA_ENABLE_DEPRECATION_WARNINGS
 	}
 
 	return false;
diff --git a/Engine/Plugins/Compositing/OpenColorIO/Source/OpenColorIO/Private/OpenColorIODerivedDataVersion.h b/Engine/Plugins/Compositing/OpenColorIO/Source/OpenColorIO/Private/OpenColorIODerivedDataVersion.h
--- a/Engine/Plugins/Compositing/OpenColorIO/Source/OpenColorIO/Private/OpenColorIODerivedDataVersion.h
+++ b/Engine/Plugins/Compositing/OpenColorIO/Source/OpenColorIO/Private/OpenColorIODerivedDataVersion.h
@@ -10,4 +10,4 @@ OpenColorIOShaderDerivedDataVersion.h: Derived data version for OpenColorIO shad
 
 // In case of merge conflicts with DDC versions, you MUST generate a new GUID and set this new
 // guid as version
-#define OPENCOLORIO_DERIVEDDATA_VER		TEXT("2018266A-4ABF-4AA9-B1B4-695931DAAFF7")
\ No newline at end of file
+#define OPENCOLORIO_DERIVEDDATA_VER		TEXT("1D868D37-BB45-4052-A0D7-85FFDD870812")
\ No newline at end of file
diff --git a/Engine/Plugins/Compositing/OpenColorIO/Source/OpenColorIO/Private/OpenColorIORendering.cpp b/Engine/Plugins/Compositing/OpenColorIO/Source/OpenColorIO/Private/OpenColorIORendering.cpp
--- a/Engine/Plugins/Compositing/OpenColorIO/Source/OpenColorIO/Private/OpenColorIORendering.cpp
+++ b/Engine/Plugins/Compositing/OpenColorIO/Source/OpenColorIO/Private/OpenColorIORendering.cpp
@@ -18,6 +18,7 @@
 #include "OpenColorIOColorTransform.h"
 #include "SceneInterface.h"
 #include "ScreenPass.h"
+#include "ShaderParameterStruct.h"
 #include "TextureResource.h"
 
 
@@ -42,19 +43,20 @@ void FOpenColorIORendering::AddPass_RenderThread(
 	FGlobalShaderMap* ShaderMap = GetGlobalShaderMap(FeatureLevel);
 	TShaderMapRef<FScreenPassVS> VertexShader(ShaderMap);
 
-	bool bRenderPassSuccessed = false;
-
-	TSortedMap<int32, FTextureResource*> TextureResources;
+	TArray<FTextureResource*> TextureResources;
 	TextureResources.Reserve(InPassResource.TextureResources.Num());
 
-	for (const TPair<int32, TWeakObjectPtr<UTexture>>& Pair : InPassResource.TextureResources)
+	for (const TWeakObjectPtr<UTexture>& Texture : InPassResource.TextureResources)
 	{
-		if (Pair.Value.IsValid(false, true))
+		if (Texture.IsValid(false, true))
 		{
 			// UTexture's GetResource() is safe to call from the render thread.
-			if (FTextureResource* TextureResource = Pair.Value->GetResource())
+			if (FTextureResource* TextureResource = Texture->GetResource())
 			{
-				TextureResources.Add(Pair.Key, TextureResource);
+				if (TextureResource->TextureRHI.IsValid())
+				{
+					TextureResources.Add(TextureResource);
+				}
 			}
 		}
 	}
@@ -63,25 +65,68 @@ void FOpenColorIORendering::AddPass_RenderThread(
 
 	if (bAllExpectedTexturesResources && InPassResource.ShaderResource != nullptr)
 	{
-		TShaderRef<FOpenColorIOPixelShader> OCIOPixelShader = InPassResource.ShaderResource->GetShader<FOpenColorIOPixelShader>();
-
-		FOpenColorIOPixelShaderParameters* Parameters = GraphBuilder.AllocParameters<FOpenColorIOPixelShaderParameters>();
+		TShaderRef<FOpenColorIOPixelShader> PixelShader = InPassResource.ShaderResource->GetShader<FOpenColorIOPixelShader>();
+		const FShaderParametersMetadata* ParametersMetadata = InPassResource.ShaderResource->GetShaderParamMetadata();
+		FOpenColorIOPixelShader::FParameters* PixelShaderParameters = GraphBuilder.AllocParameters<FOpenColorIOPixelShader::FParameters>(ParametersMetadata);
+		
+		uint8* ParameterBufferPtr = reinterpret_cast<uint8*>(PixelShaderParameters);
+		const TArray<FShaderParametersMetadata::FMember>& ParamMembers = ParametersMetadata->GetMembers();
+
+		FOpenColorIOPixelShaderParameters* Parameters = reinterpret_cast<FOpenColorIOPixelShaderParameters*>(ParameterBufferPtr);
 		Parameters->InputTexture = Input.Texture;
 		Parameters->InputTextureSampler = TStaticSamplerState<>::GetRHI();
-		if (OpenColorIOBindTextureResources(Parameters, TextureResources))
+		Parameters->Gamma = InGamma;
+		Parameters->TransformAlpha = static_cast<uint32>(TransformAlpha);
+		Parameters->RenderTargets[0] = Output.GetRenderTargetBinding();
+		
+		// Dynamic texture parameters (generated by OpenColorIO) are appended after our built-in struct of parameters.
+		const int32 StartMemberTextureIndex = 1;
+		for (int32 MemberIndex = StartMemberTextureIndex; MemberIndex < ParamMembers.Num(); ++MemberIndex)
 		{
-			Parameters->Gamma = InGamma;
-			Parameters->TransformAlpha = (uint32)TransformAlpha;
-			Parameters->RenderTargets[0] = Output.GetRenderTargetBinding();
+			const int32 TextureIndex = (MemberIndex - StartMemberTextureIndex) / 2;
 
-			FRDGEventName PassName = RDG_EVENT_NAME("OpenColorIOPass %dx%d (%s)", Output.ViewRect.Width(), Output.ViewRect.Height(), InPassResource.TransformName.IsEmpty() ? TEXT("Unspecified Transform") : *InPassResource.TransformName);
-			AddDrawScreenPass(GraphBuilder, MoveTemp(PassName), ViewInfo, OutputViewport, InputViewport, VertexShader, OCIOPixelShader, Parameters);
-
-			bRenderPassSuccessed = true;
+			const FShaderParametersMetadata::FMember& Member = ParamMembers[MemberIndex];
+			switch (Member.GetBaseType())
+			{
+			case UBMT_TEXTURE:
+				*((FRHITexture**)(ParameterBufferPtr + Member.GetOffset())) = TextureResources[TextureIndex]->TextureRHI.GetReference();
+				break;
+			case UBMT_SAMPLER:
+				*((FRHISamplerState**)(ParameterBufferPtr + Member.GetOffset())) = TextureResources[TextureIndex]->SamplerStateRHI.GetReference();
+				break;
+
+			default:
+				checkNoEntry();
+				break;
+			}
 		}
-	}
 
-	if(!bRenderPassSuccessed)
+		check(VertexShader.IsValid());
+		check(PixelShader.IsValid());
+		check(PixelShaderParameters);
+
+		ClearUnusedGraphResources(PixelShader, ParametersMetadata, PixelShaderParameters);
+
+		FRHIBlendState* BlendState = FScreenPassPipelineState::FDefaultBlendState::GetRHI();
+		FRHIDepthStencilState* DepthStencilState = FScreenPassPipelineState::FDefaultDepthStencilState::GetRHI();
+		const FScreenPassPipelineState PipelineState(VertexShader, PixelShader, BlendState, DepthStencilState);
+
+		FRDGEventName PassName = RDG_EVENT_NAME("OpenColorIOPass %dx%d (%s)", Output.ViewRect.Width(), Output.ViewRect.Height(), InPassResource.TransformName.IsEmpty() ? TEXT("Unspecified Transform") : *InPassResource.TransformName);
+		GraphBuilder.AddPass(
+			Forward<FRDGEventName&&>(PassName),
+			ParametersMetadata,
+			PixelShaderParameters,
+			ERDGPassFlags::Raster,
+			[ViewInfo, OutputViewport, InputViewport, PipelineState, ParametersMetadata, PixelShaderParameters](FRDGAsyncTask, FRHICommandList& RHICmdList)
+			{
+				constexpr EScreenPassDrawFlags Flags = EScreenPassDrawFlags::None;
+				DrawScreenPass(RHICmdList, ViewInfo, OutputViewport, InputViewport, PipelineState, Flags, [&](FRHICommandList&)
+					{
+						SetShaderParameters(RHICmdList, PipelineState.PixelShader.GetPixelShader(), PipelineState.PixelShader->Bindings, ParametersMetadata, PixelShaderParameters);
+					});
+			});
+	}
+	else
 	{
 		// Fallback pass, printing invalid message across the viewport.
 		TShaderMapRef<FOpenColorIOInvalidPixelShader> OCIOInvalidPixelShader(ShaderMap);
@@ -191,7 +236,7 @@ bool FOpenColorIORendering::ApplyColorTransform(UWorld* InWorld, const FOpenColo
 
 	const ERHIFeatureLevel::Type FeatureLevel = InWorld->Scene->GetFeatureLevel();
 	FOpenColorIOTransformResource* ShaderResource = nullptr;
-	TSortedMap<int32, TWeakObjectPtr<UTexture>> TransformTextures;
+	TArray<TWeakObjectPtr<UTexture>> TransformTextures;
 
 	if (InSettings.ConfigurationSource != nullptr)
 	{
diff --git a/Engine/Plugins/Compositing/OpenColorIO/Source/OpenColorIO/Private/OpenColorIOShader.cpp b/Engine/Plugins/Compositing/OpenColorIO/Source/OpenColorIO/Private/OpenColorIOShader.cpp
--- a/Engine/Plugins/Compositing/OpenColorIO/Source/OpenColorIO/Private/OpenColorIOShader.cpp
+++ b/Engine/Plugins/Compositing/OpenColorIO/Source/OpenColorIO/Private/OpenColorIOShader.cpp
@@ -2,48 +2,28 @@
 
 #include "OpenColorIOShader.h"
 
-#include "Containers/SortedMap.h"
 #include "OpenColorIOModule.h"
-#include "OpenColorIOShared.h"
-#include "TextureResource.h"
 #include "SystemTextures.h"
 
 bool OpenColorIOBindTextureResources(FOpenColorIOPixelShaderParameters* Parameters, const TSortedMap<int32, FTextureResource*>& InTextureResources)
 {
-	/**
-	* The engine's shader parameter API doesn't provide name/slot-based updates (in the style of legacy OpenGL),
-	* which is what OCIOv2 apparently relies on. Knowing that our input texture array follows the shader
-	* binding slot order, we set the FOpenColorIOPixelShaderParameters members by copying the relevant
-	* pointers using offsets from the initial "0" textures.
-	* 
-	* @todo: Improve binding logic, maybe using an approach similar to shaders generated by ComputeGraph.
-	*/
-	const uint8* BaseAddress = (uint8*)Parameters;
-
-	for ( const TPair<int32, FTextureResource*>& Pair : InTextureResources )
-	{
-		const int32 Index = Pair.Key;
-		if (!ensureMsgf(Index >= 0 && Index < (int32)OpenColorIOShader::MaximumTextureSlots, TEXT("Out of range LUT texture slot index.")))
-		{
-			return false;
-		}
-
-		const FTextureResource* Texture = Pair.Value;
-		if (!ensureMsgf(Texture, TEXT("Texture resource should not be null.")))
-		{
-			return false;
-		}
-
-		const bool bIsLUT3D = Texture->GetTexture3DResource() != nullptr;
-
-		const uint32 BaseOffset = bIsLUT3D ? offsetof(FOpenColorIOPixelShaderParameters, Ocio_lut3d_0) : offsetof(FOpenColorIOPixelShaderParameters, Ocio_lut1d_0);
-		const uint32 ParamOffset = BaseOffset + (sizeof(FRHITexture*) + sizeof(FRHISamplerState*)) * Index;
-
-		*((FRHITexture**)(BaseAddress + ParamOffset)) = Texture->TextureRHI.GetReference();
-		*((FRHISamplerState**)(BaseAddress + ParamOffset + sizeof(FRHITexture*))) = Texture->SamplerStateRHI.GetReference();
-	}
+	return false;
+}
 
-	return true;
+FOpenColorIOPixelShader::FOpenColorIOPixelShader(const FOpenColorIOShaderType::CompiledShaderInitializerType& Initializer)
+	: FGlobalShader(Initializer)
+{
+	constexpr bool bShouldBindEverything = true;
+	const FShaderParametersMetadata& ShaderParametersMetadata =
+		static_cast<const FOpenColorIOShaderType::FParameters*>(Initializer.Parameters)->ShaderParamMetadata;
+
+	Bindings.BindForLegacyShaderParameters(
+		this,
+		Initializer.PermutationId,
+		Initializer.ParameterMap,
+		ShaderParametersMetadata,
+		bShouldBindEverything
+	);
 }
 
 FRHITexture* OpenColorIOGetMiniFontTexture()
diff --git a/Engine/Plugins/Compositing/OpenColorIO/Source/OpenColorIO/Private/OpenColorIOShaderMap.cpp b/Engine/Plugins/Compositing/OpenColorIO/Source/OpenColorIO/Private/OpenColorIOShaderMap.cpp
--- a/Engine/Plugins/Compositing/OpenColorIO/Source/OpenColorIO/Private/OpenColorIOShaderMap.cpp
+++ b/Engine/Plugins/Compositing/OpenColorIO/Source/OpenColorIO/Private/OpenColorIOShaderMap.cpp
@@ -163,7 +163,7 @@ void FOpenColorIOShaderType::BeginCompileShader(
 	{
 		return;
 	}
-
+	NewJob->ShaderParameters = MakeShared<const FParameters, ESPMode::ThreadSafe>(*InColorTransform->GetShaderParamMetadata());
 	NewJob->Input.SharedEnvironment = InCompilationEnvironment;
 	NewJob->Input.Target = InTarget;
 	NewJob->Input.ShaderFormat = LegacyShaderPlatformToShaderFormat(InPlatform);
@@ -204,8 +204,16 @@ FShader* FOpenColorIOShaderType::FinishCompileShader(
 {
 	check(InCurrentJob.bSucceeded);
 
-	const int32 PermutationId = 0;
-	FShader* Shader = ConstructCompiled(FOpenColorIOShaderType::CompiledShaderInitializerType(this, PermutationId, InCurrentJob.Output, InShaderMapHash, InDebugDescription));
+	FShader* Shader = ConstructCompiled(
+		FOpenColorIOShaderType::CompiledShaderInitializerType(
+			this,
+			static_cast<const FParameters*>(InCurrentJob.ShaderParameters.Get()),
+			InCurrentJob.Key.PermutationId,
+			InCurrentJob.Output,
+			InShaderMapHash,
+			InDebugDescription
+		)
+	);
 	InCurrentJob.Output.ParameterMap.VerifyBindingsAreComplete(GetName(), InCurrentJob.Output.Target, nullptr);
 
 	return Shader;
diff --git a/Engine/Plugins/Compositing/OpenColorIO/Source/OpenColorIO/Private/OpenColorIOShared.cpp b/Engine/Plugins/Compositing/OpenColorIO/Source/OpenColorIO/Private/OpenColorIOShared.cpp
--- a/Engine/Plugins/Compositing/OpenColorIO/Source/OpenColorIO/Private/OpenColorIOShared.cpp
+++ b/Engine/Plugins/Compositing/OpenColorIO/Source/OpenColorIO/Private/OpenColorIOShared.cpp
@@ -209,7 +209,7 @@ void FOpenColorIOTransformResource::SerializeShaderMap(FArchive& Ar)
 	}
 }
 
-void FOpenColorIOTransformResource::SetupResource(ERHIFeatureLevel::Type InFeatureLevel, const FString& InShaderCodeHash, const FString& InShadercode, const FString& InRawConfigHash, const FString& InFriendlyName, const FName& InAssetPath)
+void FOpenColorIOTransformResource::SetupResource(ERHIFeatureLevel::Type InFeatureLevel, const FString& InShaderCodeHash, const FString& InShadercode, const FString& InRawConfigHash, const FString& InFriendlyName, const FName& InAssetPath, TUniquePtr<FShaderParametersMetadata> InShaderParameterMetadata)
 {
 	// When this happens we assume that shader was cooked and we don't need to do anything. 
 	if (InShaderCodeHash.IsEmpty() && InRawConfigHash.IsEmpty())
@@ -225,6 +225,7 @@ void FOpenColorIOTransformResource::SetupResource(ERHIFeatureLevel::Type InFeatu
 	ShaderCode = InShadercode;
 	FriendlyName = InFriendlyName.Replace(TEXT("/"), TEXT(""));
 	AssetPath = InAssetPath;
+	ShaderParametersMetadata = MoveTemp(InShaderParameterMetadata);
 
 	SetFeatureLevel(InFeatureLevel);
 }
@@ -368,6 +369,11 @@ bool FOpenColorIOTransformResource::CacheShaders(const FOpenColorIOShaderMapId&
 	return bSucceeded;
 }
 
+const FShaderParametersMetadata* FOpenColorIOTransformResource::GetShaderParamMetadata() const
+{
+	return ShaderParametersMetadata.Get();
+}
+
 #if WITH_EDITOR
 void FOpenColorIOTransformResource::FinishCompilation()
 {
diff --git a/Engine/Plugins/Compositing/OpenColorIO/Source/OpenColorIO/Public/OpenColorIOColorTransform.h b/Engine/Plugins/Compositing/OpenColorIO/Source/OpenColorIO/Public/OpenColorIOColorTransform.h
--- a/Engine/Plugins/Compositing/OpenColorIO/Source/OpenColorIO/Public/OpenColorIOColorTransform.h
+++ b/Engine/Plugins/Compositing/OpenColorIO/Source/OpenColorIO/Public/OpenColorIOColorTransform.h
@@ -31,7 +31,10 @@ class OPENCOLORIO_API UOpenColorIOColorTransform : public UObject
 public:
 
 	UOpenColorIOColorTransform(const FObjectInitializer& ObjectInitializer);
+	
+	PRAGMA_DISABLE_DEPRECATION_WARNINGS
 	virtual ~UOpenColorIOColorTransform() {};
+	PRAGMA_ENABLE_DEPRECATION_WARNINGS
 
 public:
 
@@ -44,6 +47,9 @@ public:
 	/**
 	 * Returns the desired resources required to apply this transform during rendering.
 	 */
+	bool GetRenderResources(ERHIFeatureLevel::Type InFeatureLevel, FOpenColorIOTransformResource*& OutShaderResource, TArray<TWeakObjectPtr<UTexture>>& OutTextureResources) const;
+
+	UE_DEPRECATED(5.6, "This method is deprecated, please use the one with an array of textures instead.")
 	bool GetRenderResources(ERHIFeatureLevel::Type InFeatureLevel, FOpenColorIOTransformResource*& OutShaderResource, TSortedMap<int32, TWeakObjectPtr<UTexture>>& OutTextureResources) const;
 
 	/**
@@ -145,8 +151,12 @@ private:
 	void ProcessTransformForGPU();
 #endif
 
+	/* Release all shader maps. */
 	void FlushResourceShaderMaps();
 
+	/* Create the shader parameter metadata, using our built-in struct and any LUT-associated texture and sampler resources. */
+	TUniquePtr<FShaderParametersMetadata> BuildShaderParamMetadata();
+
 public:
 
 	//~ Begin UObject interface
@@ -205,9 +215,15 @@ public:
 	EOpenColorIOViewTransformDirection DisplayViewDirection = EOpenColorIOViewTransformDirection::Forward;
 private:
 
-	/** If the color space requires textures, this will contains the data to do the transform */
+#if WITH_EDITORONLY_DATA
+	UE_DEPRECATED(5.6, "Map of textures has been deprecated in favor of an array.")
+	UPROPERTY(meta = (DeprecatedProperty, DeprecationMessage = "Map of textures has been deprecated in favor of an array."))
+	TMap<int32, TObjectPtr<UTexture>> Textures_DEPRECATED;
+#endif
+
+	/** If the color space transform requires LUTs, this will contains the texture data to do the transform */
 	UPROPERTY()
-	TMap<int32, TObjectPtr<UTexture>> Textures;
+	TArray<TObjectPtr<UTexture>> LookupTextures;
 
 #if WITH_EDITORONLY_DATA
 	/** Generated transform shader hash. */
diff --git a/Engine/Plugins/Compositing/OpenColorIO/Source/OpenColorIO/Public/OpenColorIOConfiguration.h b/Engine/Plugins/Compositing/OpenColorIO/Source/OpenColorIO/Public/OpenColorIOConfiguration.h
--- a/Engine/Plugins/Compositing/OpenColorIO/Source/OpenColorIO/Public/OpenColorIOConfiguration.h
+++ b/Engine/Plugins/Compositing/OpenColorIO/Source/OpenColorIO/Public/OpenColorIOConfiguration.h
@@ -34,14 +34,28 @@ public:
 
 public:
 
+	/** Check if the transform resources (shader & optionally lookup textures) are ready for use. */
 	bool IsTransformReady(const FOpenColorIOColorConversionSettings& InSettings);
 
+	/** Get the shader and optionally lookup texture resources to be used by the color transform render pass specified by the settings parameter. */
+	bool GetRenderResources(ERHIFeatureLevel::Type InFeatureLevel, const FOpenColorIOColorConversionSettings& InSettings, FOpenColorIOTransformResource*& OutShaderResource, TArray<TWeakObjectPtr<UTexture>>& OutTextureResources);
+
+	UE_DEPRECATED(5.6, "This method is deprecated, please use the one with an array of textures instead.")
 	bool GetRenderResources(ERHIFeatureLevel::Type InFeatureLevel, const FOpenColorIOColorConversionSettings& InSettings, FOpenColorIOTransformResource*& OutShaderResource, TSortedMap<int32, TWeakObjectPtr<UTexture>>& OutTextureResources);
 	
+	/** Check if this configuration contains a transform from the specified source & destination color spaces. */
 	bool HasTransform(const FString& InSourceColorSpace, const FString& InDestinationColorSpace);
+	
+	/** Check if this configuration contains a transform from the specified source color space & view destination. */
 	bool HasTransform(const FString& InSourceColorSpace, const FString& InDisplay, const FString& InView, EOpenColorIOViewTransformDirection InDirection);
+
+	/** Check if this configuration contains a color space. */
 	bool HasDesiredColorSpace(const FOpenColorIOColorSpace& ColorSpace) const;
+
+	/** Check if this configuration contains a display-view. */
 	bool HasDesiredDisplayView(const FOpenColorIODisplayView& DisplayView) const;
+
+	/** Check if this configuration is valid. (TODO: Rename to avoid confusion with config validation calls provided by the library.) */
 	bool Validate() const;
 
 	/** Apply the color transform in-place to the specified image. */
diff --git a/Engine/Plugins/Compositing/OpenColorIO/Source/OpenColorIO/Public/OpenColorIORendering.h b/Engine/Plugins/Compositing/OpenColorIO/Source/OpenColorIO/Public/OpenColorIORendering.h
--- a/Engine/Plugins/Compositing/OpenColorIO/Source/OpenColorIO/Public/OpenColorIORendering.h
+++ b/Engine/Plugins/Compositing/OpenColorIO/Source/OpenColorIO/Public/OpenColorIORendering.h
@@ -26,7 +26,7 @@ struct OPENCOLORIO_API FOpenColorIORenderPassResources
 	FOpenColorIOTransformResource* ShaderResource = nullptr;
 	
 	/** Collection of LUT textures needed by the shader. */
-	TSortedMap<int32, TWeakObjectPtr<UTexture>> TextureResources = {};
+	TArray<TWeakObjectPtr<UTexture>> TextureResources = {};
 
 	/** Color transform string description. */
 	FString TransformName = FString();
diff --git a/Engine/Plugins/Compositing/OpenColorIO/Source/OpenColorIO/Public/OpenColorIOShader.h b/Engine/Plugins/Compositing/OpenColorIO/Source/OpenColorIO/Public/OpenColorIOShader.h
--- a/Engine/Plugins/Compositing/OpenColorIO/Source/OpenColorIO/Public/OpenColorIOShader.h
+++ b/Engine/Plugins/Compositing/OpenColorIO/Source/OpenColorIO/Public/OpenColorIOShader.h
@@ -16,7 +16,6 @@
 #include "ShaderParameterStruct.h"
 
 class FTextureResource;
-class UClass;
 class FRHITexture;
 
 BEGIN_SHADER_PARAMETER_STRUCT(FOpenColorIOPixelShaderParameters, )
@@ -24,69 +23,28 @@ BEGIN_SHADER_PARAMETER_STRUCT(FOpenColorIOPixelShaderParameters, )
 	SHADER_PARAMETER_SAMPLER(SamplerState, InputTextureSampler)
 	SHADER_PARAMETER(float, Gamma)
 	SHADER_PARAMETER(uint32, TransformAlpha)
-	
-	SHADER_PARAMETER_TEXTURE(Texture3D, Ocio_lut3d_0)
-	SHADER_PARAMETER_SAMPLER(SamplerState, Ocio_lut3d_0Sampler)
-	SHADER_PARAMETER_TEXTURE(Texture3D, Ocio_lut3d_1)
-	SHADER_PARAMETER_SAMPLER(SamplerState, Ocio_lut3d_1Sampler)
-	SHADER_PARAMETER_TEXTURE(Texture3D, Ocio_lut3d_2)
-	SHADER_PARAMETER_SAMPLER(SamplerState, Ocio_lut3d_2Sampler)
-	SHADER_PARAMETER_TEXTURE(Texture3D, Ocio_lut3d_3)
-	SHADER_PARAMETER_SAMPLER(SamplerState, Ocio_lut3d_3Sampler)
-	SHADER_PARAMETER_TEXTURE(Texture3D, Ocio_lut3d_4)
-	SHADER_PARAMETER_SAMPLER(SamplerState, Ocio_lut3d_4Sampler)
-	SHADER_PARAMETER_TEXTURE(Texture3D, Ocio_lut3d_5)
-	SHADER_PARAMETER_SAMPLER(SamplerState, Ocio_lut3d_5Sampler)
-	SHADER_PARAMETER_TEXTURE(Texture3D, Ocio_lut3d_6)
-	SHADER_PARAMETER_SAMPLER(SamplerState, Ocio_lut3d_6Sampler)
-	SHADER_PARAMETER_TEXTURE(Texture3D, Ocio_lut3d_7)
-	SHADER_PARAMETER_SAMPLER(SamplerState, Ocio_lut3d_7Sampler)
-
-	SHADER_PARAMETER_TEXTURE(Texture2D, Ocio_lut1d_0)
-	SHADER_PARAMETER_SAMPLER(SamplerState, Ocio_lut1d_0Sampler)
-	SHADER_PARAMETER_TEXTURE(Texture2D, Ocio_lut1d_1)
-	SHADER_PARAMETER_SAMPLER(SamplerState, Ocio_lut1d_1Sampler)
-	SHADER_PARAMETER_TEXTURE(Texture2D, Ocio_lut1d_2)
-	SHADER_PARAMETER_SAMPLER(SamplerState, Ocio_lut1d_2Sampler)
-	SHADER_PARAMETER_TEXTURE(Texture2D, Ocio_lut1d_3)
-	SHADER_PARAMETER_SAMPLER(SamplerState, Ocio_lut1d_3Sampler)
-	SHADER_PARAMETER_TEXTURE(Texture2D, Ocio_lut1d_4)
-	SHADER_PARAMETER_SAMPLER(SamplerState, Ocio_lut1d_4Sampler)
-	SHADER_PARAMETER_TEXTURE(Texture2D, Ocio_lut1d_5)
-	SHADER_PARAMETER_SAMPLER(SamplerState, Ocio_lut1d_5Sampler)
-	SHADER_PARAMETER_TEXTURE(Texture2D, Ocio_lut1d_6)
-	SHADER_PARAMETER_SAMPLER(SamplerState, Ocio_lut1d_6Sampler)
-	SHADER_PARAMETER_TEXTURE(Texture2D, Ocio_lut1d_7)
-	SHADER_PARAMETER_SAMPLER(SamplerState, Ocio_lut1d_7Sampler)
-	
 	RENDER_TARGET_BINDING_SLOTS()
 END_SHADER_PARAMETER_STRUCT()
 
 namespace OpenColorIOShader
 {
-/** Fields moved to OpenColorIOWrapper.h
- *
- *	static constexpr TCHAR OpenColorIOShaderFunctionName[] = TEXT("OCIOConvert");
- *	static constexpr uint32 Lut3dEdgeLength = 65;
-*/
-
-	// Combined maximum of either Ocio_lut1d or Ocio_lut3d, since the generated shaders will never use the same slots for both types.
+	UE_DEPRECATED(5.6, "This limit is now deprecated.")
 	static constexpr uint32 MaximumTextureSlots = 8;
 }
 
 class FOpenColorIOPixelShader : public FGlobalShader
 {
-public:
 	DECLARE_EXPORTED_SHADER_TYPE(FOpenColorIOPixelShader, OpenColorIO, OPENCOLORIO_API);
-	SHADER_USE_PARAMETER_STRUCT(FOpenColorIOPixelShader, FGlobalShader);
 
-	using FParameters = FOpenColorIOPixelShaderParameters;
-	using FPermutationParameters = FOpenColorIOShaderPermutationParameters;
+public:
+	FOpenColorIOPixelShader() = default;
+	FOpenColorIOPixelShader(
+		const FOpenColorIOShaderType::CompiledShaderInitializerType& Initializer
+	);
 
-	static bool ShouldCompilePermutation(const FOpenColorIOShaderPermutationParameters& Parameters)
+	MS_ALIGN(SHADER_PARAMETER_STRUCT_ALIGNMENT) struct FParameters
 	{
-		return true;
-	}
+	};
 };
 
 BEGIN_SHADER_PARAMETER_STRUCT(FOpenColorIOInvalidShaderParameters, )
@@ -105,6 +63,7 @@ public:
 	using FParameters = FOpenColorIOInvalidShaderParameters;
 };
 
+UE_DEPRECATED(5.6, "This shader parameter binding function is now deprecated.")
 OPENCOLORIO_API bool OpenColorIOBindTextureResources(FOpenColorIOPixelShaderParameters* Parameters, const TSortedMap<int32, FTextureResource*>& InTextureResources);
 
 FRHITexture* OpenColorIOGetMiniFontTexture();
diff --git a/Engine/Plugins/Compositing/OpenColorIO/Source/OpenColorIO/Public/OpenColorIOShaderType.h b/Engine/Plugins/Compositing/OpenColorIO/Source/OpenColorIO/Public/OpenColorIOShaderType.h
--- a/Engine/Plugins/Compositing/OpenColorIO/Source/OpenColorIO/Public/OpenColorIOShaderType.h
+++ b/Engine/Plugins/Compositing/OpenColorIO/Source/OpenColorIO/Public/OpenColorIOShaderType.h
@@ -46,18 +46,31 @@ struct FOpenColorIOShaderPermutationParameters : public FGlobalShaderPermutation
 class FOpenColorIOShaderType : public FShaderType
 {
 public:
+	struct FParameters : public FShaderType::FParameters
+	{
+		const FShaderParametersMetadata& ShaderParamMetadata;
+
+		FParameters(
+			const FShaderParametersMetadata& InShaderParamMetadata
+		)
+			: ShaderParamMetadata(InShaderParamMetadata)
+		{
+		}
+	};
+
 	struct CompiledShaderInitializerType : FShaderCompiledShaderInitializerType
 	{
 		const FString DebugDescription;
 
 		CompiledShaderInitializerType(
 			const FShaderType* InType,
+			const FParameters* InParameters,
 			int32 InPermutationId,
 			const FShaderCompilerOutput& CompilerOutput,
 			const FSHAHash& InOCIOShaderMapHash,
 			const FString& InDebugDescription
 			)
-		: FShaderCompiledShaderInitializerType(InType,nullptr,InPermutationId,CompilerOutput, InOCIOShaderMapHash,nullptr,nullptr)
+		: FShaderCompiledShaderInitializerType(InType, InParameters, InPermutationId, CompilerOutput, InOCIOShaderMapHash, nullptr, nullptr)
 		, DebugDescription(InDebugDescription)
 		{}
 	};
diff --git a/Engine/Plugins/Compositing/OpenColorIO/Source/OpenColorIO/Public/OpenColorIOShared.h b/Engine/Plugins/Compositing/OpenColorIO/Source/OpenColorIO/Public/OpenColorIOShared.h
--- a/Engine/Plugins/Compositing/OpenColorIO/Source/OpenColorIO/Public/OpenColorIOShared.h
+++ b/Engine/Plugins/Compositing/OpenColorIO/Source/OpenColorIO/Public/OpenColorIOShared.h
@@ -335,6 +335,9 @@ public:
 	bool CacheShaders(EShaderPlatform InPlatform, const ITargetPlatform* TargetPlatform, bool bApplyCompletedShaderMapForRendering, bool bSynchronous);
 	bool CacheShaders(const FOpenColorIOShaderMapId& InShaderMapId, EShaderPlatform InPlatform, bool bApplyCompletedShaderMapForRendering, bool bSynchronous);
 
+	/** Returns the shader parameter metadata for the generated shader (built-in FOpenColorIOPixelShaderParameters, followed by dynamic textures & samplers). */
+	const FShaderParametersMetadata* GetShaderParamMetadata() const;
+
 	/**
 	 * Should the shader for this color transform with the given platform, shader type and vertex
 	 * factory type combination be compiled
@@ -458,7 +461,8 @@ public:
 		const FString& InShadercode,
 		const FString& InRawConfigHash,
 		const FString& InFriendlyName,
-		const FName& InAssetPath
+		const FName& InAssetPath,
+		TUniquePtr<FShaderParametersMetadata> InShaderParameterMetadata
 	);
 
 	void SetCompileErrors(TArray<FString> &InErrors)
@@ -519,6 +523,11 @@ private:
 	 */
 	TRefCountPtr<FOpenColorIOShaderMap> GameThreadShaderMap;
 
+	/**
+	 * The shader parameter metadata for the generated shader (built-in FOpenColorIOPixelShaderParameters, followed by dynamic textures & samplers).
+	 */
+	TUniquePtr<FShaderParametersMetadata> ShaderParametersMetadata;
+
 	/** 
 	 * Shader map for this FOpenColorIOTransformResource which is accessible by the rendering thread. 
 	 * This must be updated along with GameThreadShaderMap, but on the rendering thread.
