OCIO: Add drag and drop support to .ocio config files.

diff --git a/Engine/Plugins/Compositing/OpenColorIO/Source/OpenColorIOEditor/Private/Factories/OpenColorIOConfigurationFactory.cpp b/Engine/Plugins/Compositing/OpenColorIO/Source/OpenColorIOEditor/Private/Factories/OpenColorIOConfigurationFactory.cpp
new file mode 100644
--- /dev/null
+++ b/Engine/Plugins/Compositing/OpenColorIO/Source/OpenColorIOEditor/Private/Factories/OpenColorIOConfigurationFactory.cpp
@@ -0,0 +1,55 @@
+// Copyright Epic Games, Inc. All Rights Reserved.
+
+#include "OpenColorIOConfigurationFactory.h"
+
+#include "AssetTypeCategories.h"
+#include "Editor.h"
+#include "Misc/Paths.h"
+#include "OpenColorIOConfiguration.h"
+
+#include UE_INLINE_GENERATED_CPP_BY_NAME(OpenColorIOConfigurationFactory)
+
+
+/* UOpenColorIOConfigAssetFactoryNew structors
+ *****************************************************************************/
+
+UOpenColorIOConfigurationFactory::UOpenColorIOConfigurationFactory(const FObjectInitializer& ObjectInitializer)
+	: Super(ObjectInitializer)
+{
+	SupportedClass = UOpenColorIOConfiguration::StaticClass();
+	bEditorImport = true;
+	Formats.Add(TEXT("ocio;OpenColorIO Config File"));
+}
+
+
+/* UFactory overrides
+ *****************************************************************************/
+
+UObject* UOpenColorIOConfigurationFactory::FactoryCreateFile(UClass* InClass, UObject* InParent, FName InName, EObjectFlags Flags, const FString& Filename, const TCHAR* Params, FFeedbackContext* Warn, bool& bOutOperationCanceled)
+{
+	GEditor->GetEditorSubsystem<UImportSubsystem>()->BroadcastAssetPreImport(this, InClass, InParent, InName, TEXT("ocio"));
+	UOpenColorIOConfiguration* Asset = NewObject<UOpenColorIOConfiguration>(InParent, InClass, InName, Flags);
+
+	if (Asset)
+	{
+		FString ConfigFilename = CurrentFilename;
+		FPaths::MakePathRelativeTo(ConfigFilename, *FPaths::ProjectDir());
+
+		Asset->ConfigurationFile = { ConfigFilename };
+		Asset->ReloadExistingColorspaces();
+
+		GEditor->GetEditorSubsystem<UImportSubsystem>()->BroadcastAssetPostImport(this, Asset);
+	}
+	else
+	{
+		bOutOperationCanceled = true;
+	}
+
+	return Asset;
+}
+
+bool UOpenColorIOConfigurationFactory::FactoryCanImport(const FString& Filename)
+{
+	return FPaths::GetExtension(Filename) == TEXT("ocio");
+}
+
diff --git a/Engine/Plugins/Compositing/OpenColorIO/Source/OpenColorIOEditor/Private/Factories/OpenColorIOConfigurationFactory.h b/Engine/Plugins/Compositing/OpenColorIO/Source/OpenColorIOEditor/Private/Factories/OpenColorIOConfigurationFactory.h
new file mode 100644
--- /dev/null
+++ b/Engine/Plugins/Compositing/OpenColorIO/Source/OpenColorIOEditor/Private/Factories/OpenColorIOConfigurationFactory.h
@@ -0,0 +1,23 @@
+// Copyright Epic Games, Inc. All Rights Reserved.
+
+#pragma once
+
+#include "Factories/Factory.h"
+
+#include "OpenColorIOConfigurationFactory.generated.h"
+
+
+/**
+ * Implements a factory for UOpenColorIOConfiguration objects.
+ */
+UCLASS(hidecategories=Object)
+class UOpenColorIOConfigurationFactory : public UFactory
+{
+	GENERATED_UCLASS_BODY()
+
+public:
+
+	//~ UFactory Interface
+	virtual UObject* FactoryCreateFile(UClass* InClass, UObject* InParent, FName InName, EObjectFlags Flags, const FString& Filename, const TCHAR* Params, FFeedbackContext* Warn, bool& bOutOperationCanceled) override;
+	virtual bool FactoryCanImport(const FString& Filename) override;
+};
