OCIO: Tentative (quick) fix to prevent shader parameter metadata destruction during rendering, using a thread-safe shared ptr.

diff --git a/Engine/Plugins/Compositing/OpenColorIO/Source/OpenColorIO/Private/OpenColorIOColorTransform.cpp b/Engine/Plugins/Compositing/OpenColorIO/Source/OpenColorIO/Private/OpenColorIOColorTransform.cpp
--- a/Engine/Plugins/Compositing/OpenColorIO/Source/OpenColorIO/Private/OpenColorIOColorTransform.cpp
+++ b/Engine/Plugins/Compositing/OpenColorIO/Source/OpenColorIO/Private/OpenColorIOColorTransform.cpp
@@ -713,7 +713,7 @@ void UOpenColorIOColorTransform::FlushResourceShaderMaps()
 	}
 }
 
-TUniquePtr<FShaderParametersMetadata> UOpenColorIOColorTransform::BuildShaderParamMetadata()
+FShaderParametersMetadata* UOpenColorIOColorTransform::BuildShaderParamMetadata()
 {
 	// We simply use static allocations for resource names, similar to ComputeFramework::AddParamForType.
 	static TArray<FString> AllTextureResourceNames;
@@ -739,7 +739,7 @@ TUniquePtr<FShaderParametersMetadata> UOpenColorIOColorTransform::BuildShaderPar
 		}
 	}
 
-	return TUniquePtr<FShaderParametersMetadata>(Builder.Build(FShaderParametersMetadata::EUseCase::ShaderParameterStruct, TEXT("")));
+	return Builder.Build(FShaderParametersMetadata::EUseCase::ShaderParameterStruct, TEXT(""));
 }
 
 void UOpenColorIOColorTransform::PostLoad()
diff --git a/Engine/Plugins/Compositing/OpenColorIO/Source/OpenColorIO/Private/OpenColorIORendering.cpp b/Engine/Plugins/Compositing/OpenColorIO/Source/OpenColorIO/Private/OpenColorIORendering.cpp
--- a/Engine/Plugins/Compositing/OpenColorIO/Source/OpenColorIO/Private/OpenColorIORendering.cpp
+++ b/Engine/Plugins/Compositing/OpenColorIO/Source/OpenColorIO/Private/OpenColorIORendering.cpp
@@ -63,12 +63,12 @@ void FOpenColorIORendering::AddPass_RenderThread(
 	}
 
 	const bool bAllExpectedTexturesResources = TextureResources.Num() == InPassResource.TextureResources.Num();
+	const TSharedPtr<const FShaderParametersMetadata, ESPMode::ThreadSafe> ParametersMetadata = InPassResource.ShaderResource ? InPassResource.ShaderResource->GetShaderParamMetadata() : nullptr;
 
-	if (bAllExpectedTexturesResources && InPassResource.ShaderResource != nullptr)
+	if (bAllExpectedTexturesResources && InPassResource.ShaderResource != nullptr && ParametersMetadata.IsValid())
 	{
 		TShaderRef<FOpenColorIOPixelShader> PixelShader = InPassResource.ShaderResource->GetShader<FOpenColorIOPixelShader>();
-		const FShaderParametersMetadata* ParametersMetadata = InPassResource.ShaderResource->GetShaderParamMetadata();
-		FOpenColorIOPixelShader::FParameters* PixelShaderParameters = GraphBuilder.AllocParameters<FOpenColorIOPixelShader::FParameters>(ParametersMetadata);
+		FOpenColorIOPixelShader::FParameters* PixelShaderParameters = GraphBuilder.AllocParameters<FOpenColorIOPixelShader::FParameters>(ParametersMetadata.Get());
 		
 		uint8* ParameterBufferPtr = reinterpret_cast<uint8*>(PixelShaderParameters);
 		const TArray<FShaderParametersMetadata::FMember>& ParamMembers = ParametersMetadata->GetMembers();
@@ -106,7 +106,7 @@ void FOpenColorIORendering::AddPass_RenderThread(
 		check(PixelShader.IsValid());
 		check(PixelShaderParameters);
 
-		ClearUnusedGraphResources(PixelShader, ParametersMetadata, PixelShaderParameters);
+		ClearUnusedGraphResources(PixelShader, ParametersMetadata.Get(), PixelShaderParameters);
 
 		FRHIBlendState* BlendState = FScreenPassPipelineState::FDefaultBlendState::GetRHI();
 		FRHIDepthStencilState* DepthStencilState = FScreenPassPipelineState::FDefaultDepthStencilState::GetRHI();
@@ -115,7 +115,7 @@ void FOpenColorIORendering::AddPass_RenderThread(
 		FRDGEventName PassName = RDG_EVENT_NAME("OpenColorIOPass %dx%d (%s)", Output.ViewRect.Width(), Output.ViewRect.Height(), InPassResource.TransformName.IsEmpty() ? TEXT("Unspecified Transform") : *InPassResource.TransformName);
 		GraphBuilder.AddPass(
 			Forward<FRDGEventName&&>(PassName),
-			ParametersMetadata,
+			ParametersMetadata.Get(),
 			PixelShaderParameters,
 			ERDGPassFlags::Raster,
 			[ViewInfo, OutputViewport, InputViewport, PipelineState, ParametersMetadata, PixelShaderParameters](FRDGAsyncTask, FRHICommandList& RHICmdList)
@@ -123,7 +123,7 @@ void FOpenColorIORendering::AddPass_RenderThread(
 				constexpr EScreenPassDrawFlags Flags = EScreenPassDrawFlags::None;
 				DrawScreenPass(RHICmdList, ViewInfo, OutputViewport, InputViewport, PipelineState, Flags, [&](FRHICommandList&)
 					{
-						SetShaderParameters(RHICmdList, PipelineState.PixelShader.GetPixelShader(), PipelineState.PixelShader->Bindings, ParametersMetadata, PixelShaderParameters);
+						SetShaderParameters(RHICmdList, PipelineState.PixelShader.GetPixelShader(), PipelineState.PixelShader->Bindings, ParametersMetadata.Get(), PixelShaderParameters);
 					});
 			});
 	}
diff --git a/Engine/Plugins/Compositing/OpenColorIO/Source/OpenColorIO/Private/OpenColorIOShared.cpp b/Engine/Plugins/Compositing/OpenColorIO/Source/OpenColorIO/Private/OpenColorIOShared.cpp
--- a/Engine/Plugins/Compositing/OpenColorIO/Source/OpenColorIO/Private/OpenColorIOShared.cpp
+++ b/Engine/Plugins/Compositing/OpenColorIO/Source/OpenColorIO/Private/OpenColorIOShared.cpp
@@ -209,7 +209,7 @@ void FOpenColorIOTransformResource::SerializeShaderMap(FArchive& Ar)
 	}
 }
 
-void FOpenColorIOTransformResource::SetupResource(ERHIFeatureLevel::Type InFeatureLevel, const FString& InShaderCodeHash, const FString& InShadercode, const FString& InRawConfigHash, const FString& InFriendlyName, const FName& InAssetPath, TUniquePtr<FShaderParametersMetadata> InShaderParameterMetadata)
+void FOpenColorIOTransformResource::SetupResource(ERHIFeatureLevel::Type InFeatureLevel, const FString& InShaderCodeHash, const FString& InShadercode, const FString& InRawConfigHash, const FString& InFriendlyName, const FName& InAssetPath, FShaderParametersMetadata* InShaderParameterMetadata)
 {
 	// When this happens we assume that shader was cooked and we don't need to do anything. 
 	if (InShaderCodeHash.IsEmpty() && InRawConfigHash.IsEmpty())
@@ -225,7 +225,7 @@ void FOpenColorIOTransformResource::SetupResource(ERHIFeatureLevel::Type InFeatu
 	ShaderCode = InShadercode;
 	FriendlyName = InFriendlyName.Replace(TEXT("/"), TEXT(""));
 	AssetPath = InAssetPath;
-	ShaderParametersMetadata = MoveTemp(InShaderParameterMetadata);
+	ShaderParametersMetadata = TSharedPtr<FShaderParametersMetadata>(InShaderParameterMetadata);
 
 	SetFeatureLevel(InFeatureLevel);
 }
@@ -369,9 +369,9 @@ bool FOpenColorIOTransformResource::CacheShaders(const FOpenColorIOShaderMapId&
 	return bSucceeded;
 }
 
-const FShaderParametersMetadata* FOpenColorIOTransformResource::GetShaderParamMetadata() const
+TSharedPtr<const FShaderParametersMetadata, ESPMode::ThreadSafe> FOpenColorIOTransformResource::GetShaderParamMetadata() const
 {
-	return ShaderParametersMetadata.Get();
+	return ShaderParametersMetadata;
 }
 
 #if WITH_EDITOR
diff --git a/Engine/Plugins/Compositing/OpenColorIO/Source/OpenColorIO/Public/OpenColorIOColorTransform.h b/Engine/Plugins/Compositing/OpenColorIO/Source/OpenColorIO/Public/OpenColorIOColorTransform.h
--- a/Engine/Plugins/Compositing/OpenColorIO/Source/OpenColorIO/Public/OpenColorIOColorTransform.h
+++ b/Engine/Plugins/Compositing/OpenColorIO/Source/OpenColorIO/Public/OpenColorIOColorTransform.h
@@ -155,7 +155,7 @@ private:
 	void FlushResourceShaderMaps();
 
 	/* Create the shader parameter metadata, using our built-in struct and any LUT-associated texture and sampler resources. */
-	TUniquePtr<FShaderParametersMetadata> BuildShaderParamMetadata();
+	FShaderParametersMetadata* BuildShaderParamMetadata();
 
 public:
 
diff --git a/Engine/Plugins/Compositing/OpenColorIO/Source/OpenColorIO/Public/OpenColorIOShared.h b/Engine/Plugins/Compositing/OpenColorIO/Source/OpenColorIO/Public/OpenColorIOShared.h
--- a/Engine/Plugins/Compositing/OpenColorIO/Source/OpenColorIO/Public/OpenColorIOShared.h
+++ b/Engine/Plugins/Compositing/OpenColorIO/Source/OpenColorIO/Public/OpenColorIOShared.h
@@ -336,7 +336,7 @@ public:
 	bool CacheShaders(const FOpenColorIOShaderMapId& InShaderMapId, EShaderPlatform InPlatform, bool bApplyCompletedShaderMapForRendering, bool bSynchronous);
 
 	/** Returns the shader parameter metadata for the generated shader (built-in FOpenColorIOPixelShaderParameters, followed by dynamic textures & samplers). */
-	const FShaderParametersMetadata* GetShaderParamMetadata() const;
+	TSharedPtr<const FShaderParametersMetadata, ESPMode::ThreadSafe> GetShaderParamMetadata() const;
 
 	/**
 	 * Should the shader for this color transform with the given platform, shader type and vertex
@@ -455,6 +455,7 @@ public:
 
 	const FString& GetFriendlyName()	const { return FriendlyName; }
 
+	// Note: Takes ownership of the shader parameter metadata.
 	void SetupResource(
 		ERHIFeatureLevel::Type InFeatureLevel,
 		const FString& InShaderCodeHash,
@@ -462,7 +463,7 @@ public:
 		const FString& InRawConfigHash,
 		const FString& InFriendlyName,
 		const FName& InAssetPath,
-		TUniquePtr<FShaderParametersMetadata> InShaderParameterMetadata
+		FShaderParametersMetadata* InShaderParameterMetadata = nullptr
 	);
 
 	void SetCompileErrors(TArray<FString> &InErrors)
@@ -526,7 +527,7 @@ private:
 	/**
 	 * The shader parameter metadata for the generated shader (built-in FOpenColorIOPixelShaderParameters, followed by dynamic textures & samplers).
 	 */
-	TUniquePtr<FShaderParametersMetadata> ShaderParametersMetadata;
+	TSharedPtr<const FShaderParametersMetadata, ESPMode::ThreadSafe> ShaderParametersMetadata;
 
 	/** 
 	 * Shader map for this FOpenColorIOTransformResource which is accessible by the rendering thread. 
