OCIO: Added OpenColorIO.AllowAsyncCompilation cvar, as an additional workaround for recent crashes seen on Aurora.

diff --git a/Engine/Plugins/Compositing/OpenColorIO/Source/OpenColorIO/Private/OpenColorIOColorTransform.cpp b/Engine/Plugins/Compositing/OpenColorIO/Source/OpenColorIO/Private/OpenColorIOColorTransform.cpp
--- a/Engine/Plugins/Compositing/OpenColorIO/Source/OpenColorIO/Private/OpenColorIOColorTransform.cpp
+++ b/Engine/Plugins/Compositing/OpenColorIO/Source/OpenColorIO/Private/OpenColorIOColorTransform.cpp
@@ -46,6 +46,12 @@ namespace
 	}
 }
 
+static TAutoConsoleVariable<bool> CVarOpenColorIOAllowAsyncCompilation(
+	TEXT("OpenColorIO.AllowAsyncCompilation"),
+	true,
+	TEXT("When enabled, we enforce synchronous shader compilation.")
+);
+
 void UOpenColorIOColorTransform::SerializeOpenColorIOShaderMaps(const TMap<const ITargetPlatform*, TArray<FOpenColorIOTransformResource*>>* PlatformColorTransformResourcesToSavePtr, FArchive& Ar, TArray<FOpenColorIOTransformResource>&  OutLoadedResources)
 {
 	if (Ar.IsSaving())
@@ -447,6 +453,11 @@ void UOpenColorIOColorTransform::CacheResourceShadersForRendering(bool /*bRegene
 
 void UOpenColorIOColorTransform::CacheShadersForResources(EShaderPlatform InShaderPlatform, FOpenColorIOTransformResource* InResourceToCache, bool bApplyCompletedShaderMapForRendering, bool bIsSynchronous, const ITargetPlatform* TargetPlatform)
 {
+	if (!CVarOpenColorIOAllowAsyncCompilation.GetValueOnAnyThread())
+	{
+		bIsSynchronous = true;
+	}
+
 	const bool bSuccess = InResourceToCache->CacheShaders(InShaderPlatform, TargetPlatform, bApplyCompletedShaderMapForRendering, bIsSynchronous);
 
 	if (!bSuccess)
