OCIO: Deprecate serialiazed color space index.

Color space indices aren't stable in built-in transforms across library version releases, so they shouldn't be used for validation & equality checks where the names suffice.

diff --git a/Engine/Plugins/Compositing/OpenColorIO/Source/OpenColorIO/Private/OpenColorIOColorSpace.cpp b/Engine/Plugins/Compositing/OpenColorIO/Source/OpenColorIO/Private/OpenColorIOColorSpace.cpp
--- a/Engine/Plugins/Compositing/OpenColorIO/Source/OpenColorIO/Private/OpenColorIOColorSpace.cpp
+++ b/Engine/Plugins/Compositing/OpenColorIO/Source/OpenColorIO/Private/OpenColorIOColorSpace.cpp
@@ -14,13 +14,8 @@
 
 const TCHAR* FOpenColorIOColorSpace::FamilyDelimiter = TEXT("/");
 
-FOpenColorIOColorSpace::FOpenColorIOColorSpace()
-	: ColorSpaceIndex(INDEX_NONE)
-{ }
-
 FOpenColorIOColorSpace::FOpenColorIOColorSpace(const FString& InColorSpaceName, int32 InColorSpaceIndex, const FString& InFamilyName)
 	: ColorSpaceName(InColorSpaceName)
-	, ColorSpaceIndex(InColorSpaceIndex)
 	, FamilyName(InFamilyName)
 { }
 
@@ -35,12 +30,11 @@ FString FOpenColorIOColorSpace::ToString() const
 
 bool FOpenColorIOColorSpace::IsValid() const
 {
-	return ColorSpaceIndex != INDEX_NONE && !ColorSpaceName.IsEmpty();
+	return !ColorSpaceName.IsEmpty();
 }
 
 void FOpenColorIOColorSpace::Reset()
 {
-	ColorSpaceIndex = INDEX_NONE;
 	ColorSpaceName.Reset();
 	FamilyName.Reset();
 }
diff --git a/Engine/Plugins/Compositing/OpenColorIO/Source/OpenColorIO/Private/OpenColorIOConfiguration.cpp b/Engine/Plugins/Compositing/OpenColorIO/Source/OpenColorIO/Private/OpenColorIOConfiguration.cpp
--- a/Engine/Plugins/Compositing/OpenColorIO/Source/OpenColorIO/Private/OpenColorIOConfiguration.cpp
+++ b/Engine/Plugins/Compositing/OpenColorIO/Source/OpenColorIO/Private/OpenColorIOConfiguration.cpp
@@ -211,7 +211,7 @@ void UOpenColorIOConfiguration::ReloadExistingColorspaces(bool bForce)
 			// This will make sure that all colorspaces are up to date in case an index, family or name is changed.
 			for (const FOpenColorIOColorSpace& ExistingColorSpace : ColorSpacesToBeReloaded)
 			{
-				int ColorSpaceIndex = Config->GetColorSpaceIndex(*ExistingColorSpace.ColorSpaceName);
+				int32 ColorSpaceIndex = Config->GetColorSpaceIndex(*ExistingColorSpace.ColorSpaceName);
 				if (ColorSpaceIndex < 0)
 				{
 					// Name not found, therefore we don't need to re-add this colorspace.
@@ -220,7 +220,6 @@ void UOpenColorIOConfiguration::ReloadExistingColorspaces(bool bForce)
 				}
 
 				FOpenColorIOColorSpace ColorSpace;
-				ColorSpace.ColorSpaceIndex = ColorSpaceIndex;
 				ColorSpace.ColorSpaceName = ExistingColorSpace.ColorSpaceName;
 				ColorSpace.FamilyName = Config->GetColorSpaceFamilyName(*ExistingColorSpace.ColorSpaceName);
 				DesiredColorSpaces.Add(ColorSpace);
diff --git a/Engine/Plugins/Compositing/OpenColorIO/Source/OpenColorIO/Public/OpenColorIOColorSpace.h b/Engine/Plugins/Compositing/OpenColorIO/Source/OpenColorIO/Public/OpenColorIOColorSpace.h
--- a/Engine/Plugins/Compositing/OpenColorIO/Source/OpenColorIO/Public/OpenColorIOColorSpace.h
+++ b/Engine/Plugins/Compositing/OpenColorIO/Source/OpenColorIO/Public/OpenColorIOColorSpace.h
@@ -21,8 +21,16 @@ struct OPENCOLORIO_API FOpenColorIOColorSpace
 	GENERATED_BODY()
 
 public:
+
+	PRAGMA_DISABLE_DEPRECATION_WARNINGS
 	/** Default constructor. */
-	FOpenColorIOColorSpace();
+	FOpenColorIOColorSpace() = default;
+	~FOpenColorIOColorSpace() = default;
+	FOpenColorIOColorSpace(const FOpenColorIOColorSpace&) = default;
+	FOpenColorIOColorSpace(FOpenColorIOColorSpace&&) = default;
+	FOpenColorIOColorSpace& operator=(const FOpenColorIOColorSpace&) = default;
+	FOpenColorIOColorSpace& operator=(FOpenColorIOColorSpace&&) = default;
+	PRAGMA_ENABLE_DEPRECATION_WARNINGS
 
 	/**
 	 * Create and initialize a new instance.
@@ -34,8 +42,11 @@ public:
 	FString ColorSpaceName;
 
 	/** The index of the ColorSpace in the config */
-	UPROPERTY(EditAnywhere, BlueprintReadWrite, Category=ColorSpace)
-	int32 ColorSpaceIndex;
+#if WITH_EDITORONLY_DATA
+	UE_DEPRECATED(5.6, "ColorSpaceIndex has been deprecated.")
+	UPROPERTY(meta = (DeprecatedProperty, DeprecationMessage = "ColorSpaceIndex has been deprecated."))
+	int32 ColorSpaceIndex_DEPRECATED = INDEX_NONE;
+#endif
 
 	/** 
 	 * The family of this ColorSpace as specified in the configuration file. 
@@ -48,7 +59,7 @@ public:
 	static const TCHAR* FamilyDelimiter;
 
 public:
-	bool operator==(const FOpenColorIOColorSpace& Other) const { return Other.ColorSpaceIndex == ColorSpaceIndex && Other.ColorSpaceName == ColorSpaceName; }
+	bool operator==(const FOpenColorIOColorSpace& Other) const { return Other.ColorSpaceName == ColorSpaceName; }
 	bool operator!=(const FOpenColorIOColorSpace& Other) const { return !operator==(Other); }
 
 	/**
diff --git a/Engine/Plugins/Compositing/OpenColorIO/Source/OpenColorIOEditor/Private/OpenColorIOColorSpaceCustomization.cpp b/Engine/Plugins/Compositing/OpenColorIO/Source/OpenColorIOEditor/Private/OpenColorIOColorSpaceCustomization.cpp
--- a/Engine/Plugins/Compositing/OpenColorIO/Source/OpenColorIOEditor/Private/OpenColorIOColorSpaceCustomization.cpp
+++ b/Engine/Plugins/Compositing/OpenColorIO/Source/OpenColorIOEditor/Private/OpenColorIOColorSpaceCustomization.cpp
@@ -199,7 +199,6 @@ void FOpenColorIOColorSpaceCustomization::PopulateSubMenu(FMenuBuilder& InMenuBu
 		for (int32 i = 0; i < ConfigWrapper->GetNumColorSpaces(); ++i)
 		{
 			FOpenColorIOColorSpace ColorSpace;
-			ColorSpace.ColorSpaceIndex = i;
 			ColorSpace.ColorSpaceName = ConfigWrapper->GetColorSpaceName(i);
 			ColorSpace.FamilyName = ConfigWrapper->GetColorSpaceFamilyName(*ColorSpace.ColorSpaceName);
 
@@ -265,7 +264,6 @@ TSharedRef<SWidget> FOpenColorIOColorSpaceCustomization::HandleSourceComboButton
 			for (int32 i = 0; i < ColorSpaceCount; ++i)
 			{
 				FOpenColorIOColorSpace ColorSpace;
-				ColorSpace.ColorSpaceIndex = i;
 				ColorSpace.ColorSpaceName = ConfigWrapper->GetColorSpaceName(i);
 				ColorSpace.FamilyName = ConfigWrapper->GetColorSpaceFamilyName(*ColorSpace.ColorSpaceName);
 
