OCIO: Fix shader parameter metadata deallocation crash.

FOpenColorIOShaderType::FParameters used during asynchronous shader compilation would previously only hold a reference to the metadata while the transform resource could get destroyed.

diff --git a/Engine/Plugins/Compositing/OpenColorIO/Source/OpenColorIO/Private/OpenColorIOShader.cpp b/Engine/Plugins/Compositing/OpenColorIO/Source/OpenColorIO/Private/OpenColorIOShader.cpp
--- a/Engine/Plugins/Compositing/OpenColorIO/Source/OpenColorIO/Private/OpenColorIOShader.cpp
+++ b/Engine/Plugins/Compositing/OpenColorIO/Source/OpenColorIO/Private/OpenColorIOShader.cpp
@@ -14,14 +14,14 @@ FOpenColorIOPixelShader::FOpenColorIOPixelShader(const FOpenColorIOShaderType::C
 	: FGlobalShader(Initializer)
 {
 	constexpr bool bShouldBindEverything = true;
-	const FShaderParametersMetadata& ShaderParametersMetadata =
+	const TSharedRef<const FShaderParametersMetadata>& ShaderParametersMetadata =
 		static_cast<const FOpenColorIOShaderType::FParameters*>(Initializer.Parameters)->ShaderParamMetadata;
 
 	Bindings.BindForLegacyShaderParameters(
 		this,
 		Initializer.PermutationId,
 		Initializer.ParameterMap,
-		ShaderParametersMetadata,
+		ShaderParametersMetadata.Get(),
 		bShouldBindEverything
 	);
 }
diff --git a/Engine/Plugins/Compositing/OpenColorIO/Source/OpenColorIO/Private/OpenColorIOShaderMap.cpp b/Engine/Plugins/Compositing/OpenColorIO/Source/OpenColorIO/Private/OpenColorIOShaderMap.cpp
--- a/Engine/Plugins/Compositing/OpenColorIO/Source/OpenColorIO/Private/OpenColorIOShaderMap.cpp
+++ b/Engine/Plugins/Compositing/OpenColorIO/Source/OpenColorIO/Private/OpenColorIOShaderMap.cpp
@@ -164,7 +164,7 @@ void FOpenColorIOShaderType::BeginCompileShader(
 	{
 		return;
 	}
-	NewJob->ShaderParameters = MakeShared<const FParameters, ESPMode::ThreadSafe>(*InColorTransform->GetShaderParamMetadata());
+	NewJob->ShaderParameters = MakeShared<const FParameters, ESPMode::ThreadSafe>(InColorTransform->GetShaderParamMetadata().ToSharedRef());
 	NewJob->Input.SharedEnvironment = InCompilationEnvironment;
 	NewJob->Input.Target = InTarget;
 	NewJob->Input.ShaderFormat = LegacyShaderPlatformToShaderFormat(InPlatform);
diff --git a/Engine/Plugins/Compositing/OpenColorIO/Source/OpenColorIO/Private/OpenColorIOShared.cpp b/Engine/Plugins/Compositing/OpenColorIO/Source/OpenColorIO/Private/OpenColorIOShared.cpp
--- a/Engine/Plugins/Compositing/OpenColorIO/Source/OpenColorIO/Private/OpenColorIOShared.cpp
+++ b/Engine/Plugins/Compositing/OpenColorIO/Source/OpenColorIO/Private/OpenColorIOShared.cpp
@@ -227,6 +227,11 @@ void FOpenColorIOTransformResource::SetupResource(ERHIFeatureLevel::Type InFeatu
 	AssetPath = InAssetPath;
 	ShaderParametersMetadata = TSharedPtr<FShaderParametersMetadata>(InShaderParameterMetadata);
 
+	if (!ShaderParametersMetadata.IsValid())
+	{
+		UE_LOG(LogOpenColorIO, Error, TEXT("OpenColorIO shader transform resource is expecting a valid shader parameter metadata for compilation."));
+	}
+
 	SetFeatureLevel(InFeatureLevel);
 }
 
diff --git a/Engine/Plugins/Compositing/OpenColorIO/Source/OpenColorIO/Public/OpenColorIOShaderType.h b/Engine/Plugins/Compositing/OpenColorIO/Source/OpenColorIO/Public/OpenColorIOShaderType.h
--- a/Engine/Plugins/Compositing/OpenColorIO/Source/OpenColorIO/Public/OpenColorIOShaderType.h
+++ b/Engine/Plugins/Compositing/OpenColorIO/Source/OpenColorIO/Public/OpenColorIOShaderType.h
@@ -48,11 +48,9 @@ class FOpenColorIOShaderType : public FShaderType
 public:
 	struct FParameters : public FShaderType::FParameters
 	{
-		const FShaderParametersMetadata& ShaderParamMetadata;
+		TSharedRef<const FShaderParametersMetadata> ShaderParamMetadata;
 
-		FParameters(
-			const FShaderParametersMetadata& InShaderParamMetadata
-		)
+		FParameters(const TSharedRef<const FShaderParametersMetadata>& InShaderParamMetadata)
 			: ShaderParamMetadata(InShaderParamMetadata)
 		{
 		}
diff --git a/Engine/Plugins/Compositing/OpenColorIO/Source/OpenColorIO/Public/OpenColorIOShared.h b/Engine/Plugins/Compositing/OpenColorIO/Source/OpenColorIO/Public/OpenColorIOShared.h
--- a/Engine/Plugins/Compositing/OpenColorIO/Source/OpenColorIO/Public/OpenColorIOShared.h
+++ b/Engine/Plugins/Compositing/OpenColorIO/Source/OpenColorIO/Public/OpenColorIOShared.h
@@ -455,6 +455,16 @@ public:
 
 	const FString& GetFriendlyName()	const { return FriendlyName; }
 
+	UE_DEPRECATED(5.6, "This method is deprecated, please use the one with a shader parameter metadata instead.")
+	void SetupResource(
+		ERHIFeatureLevel::Type InFeatureLevel,
+		const FString& InShaderCodeHash,
+		const FString& InShadercode,
+		const FString& InRawConfigHash,
+		const FString& InFriendlyName,
+		const FName& InAssetPath
+	) { }
+
 	// Note: Takes ownership of the shader parameter metadata.
 	void SetupResource(
 		ERHIFeatureLevel::Type InFeatureLevel,
@@ -463,7 +473,7 @@ public:
 		const FString& InRawConfigHash,
 		const FString& InFriendlyName,
 		const FName& InAssetPath,
-		FShaderParametersMetadata* InShaderParameterMetadata = nullptr
+		FShaderParametersMetadata* InShaderParameterMetadata
 	);
 
 	void SetCompileErrors(TArray<FString> &InErrors)
