OCIO: Expose new default settings applied to level viewports when none was previously set, or if settings have been cleared.

* These can be easily configured via the editor settings menu.
* We also remove old logic that automatically disable invalid transforms on load, since showing the debug pattern is preferable in case an expected transform breaks due to some config changes.

diff --git a/Engine/Plugins/Compositing/OpenColorIO/Source/OpenColorIO/Public/OpenColorIOColorSpace.h b/Engine/Plugins/Compositing/OpenColorIO/Source/OpenColorIO/Public/OpenColorIOColorSpace.h
--- a/Engine/Plugins/Compositing/OpenColorIO/Source/OpenColorIO/Public/OpenColorIOColorSpace.h
+++ b/Engine/Plugins/Compositing/OpenColorIO/Source/OpenColorIO/Public/OpenColorIOColorSpace.h
@@ -240,6 +240,12 @@ public:
 
 	bool Serialize(FArchive& Ar);
 	void PostSerialize(const FArchive& Ar);
+
+	/** Determines if this DisplayConfiguration is the same as another.*/
+	bool Equals(const FOpenColorIODisplayConfiguration& Other) const
+	{
+		return bIsEnabled == Other.bIsEnabled && ColorConfiguration.Equals(Other.ColorConfiguration);
+	}
 };
 
 template<> struct TStructOpsTypeTraits<FOpenColorIODisplayConfiguration> : public TStructOpsTypeTraitsBase2<FOpenColorIODisplayConfiguration>
diff --git a/Engine/Plugins/Compositing/OpenColorIO/Source/OpenColorIOEditor/Private/OpenColorIOEditorModule.cpp b/Engine/Plugins/Compositing/OpenColorIO/Source/OpenColorIOEditor/Private/OpenColorIOEditorModule.cpp
--- a/Engine/Plugins/Compositing/OpenColorIO/Source/OpenColorIOEditor/Private/OpenColorIOEditorModule.cpp
+++ b/Engine/Plugins/Compositing/OpenColorIO/Source/OpenColorIOEditor/Private/OpenColorIOEditorModule.cpp
@@ -5,6 +5,7 @@
 #include "Engine/World.h"
 #include "Interfaces/IPluginManager.h"
 #include "IOpenColorIOModule.h"
+#include "ISettingsModule.h"
 #include "LevelEditor.h"
 #include "LevelEditorViewport.h"
 #include "Modules/ModuleManager.h"
@@ -37,12 +38,14 @@ void FOpenColorIOEditorModule::StartupModule()
 	RegisterCustomizations();
 	RegisterStyle();
 	RegisterViewMenuExtension();
+	RegisterSettings();
 }
 
 void FOpenColorIOEditorModule::ShutdownModule()
 {
 	LLM_SCOPE_BYTAG(OpenColorIO_OpenColorIOEditor);
 
+	UnregisterSettings();
 	UnregisterViewMenuExtension();
 	UnregisterStyle();
 	UnregisterCustomizations();
@@ -76,6 +79,25 @@ void FOpenColorIOEditorModule::UnregisterCustomizations()
 	}
 }
 
+void FOpenColorIOEditorModule::RegisterSettings()
+{
+	if (ISettingsModule* SettingsModule = FModuleManager::GetModulePtr<ISettingsModule>("Settings"))
+	{
+		SettingsModule->RegisterSettings("Editor", "Plugins", "OpenColorIO Editor",
+			LOCTEXT("OpenColorIOSettingsName", "OpenColorIO Editor"),
+			LOCTEXT("OpenColorIOSettingsDescription", "Configure default viewports settings for OpenColorIO."),
+			GetMutableDefault<UOpenColorIODefaultViewportSettings>());
+	}
+}
+
+void FOpenColorIOEditorModule::UnregisterSettings()
+{
+	if (ISettingsModule* SettingsModule = FModuleManager::GetModulePtr<ISettingsModule>("Settings"))
+	{
+		SettingsModule->UnregisterSettings("Editor", "Plugins", "OpenColorIO Editor");
+	}
+}
+
 void FOpenColorIOEditorModule::OnWorldInit(UWorld* InWorld, const UWorld::InitializationValues InInitializationValues)
 {
 	LLM_SCOPE_BYTAG(OpenColorIO_OpenColorIOEditor);
diff --git a/Engine/Plugins/Compositing/OpenColorIO/Source/OpenColorIOEditor/Private/OpenColorIOEditorModule.h b/Engine/Plugins/Compositing/OpenColorIO/Source/OpenColorIOEditor/Private/OpenColorIOEditorModule.h
--- a/Engine/Plugins/Compositing/OpenColorIO/Source/OpenColorIOEditor/Private/OpenColorIOEditorModule.h
+++ b/Engine/Plugins/Compositing/OpenColorIO/Source/OpenColorIOEditor/Private/OpenColorIOEditorModule.h
@@ -35,6 +35,8 @@ public:
 protected:
 	void RegisterCustomizations();
 	void UnregisterCustomizations();
+	void RegisterSettings();
+	void UnregisterSettings();
 
 	void OnWorldInit(UWorld* InWorld, const UWorld::InitializationValues InInitializationValues);
 	void OnLevelEditorFeatureLevelChanged(ERHIFeatureLevel::Type InFeatureLevel);
diff --git a/Engine/Plugins/Compositing/OpenColorIO/Source/OpenColorIOEditor/Private/OpenColorIOEditorSettings.cpp b/Engine/Plugins/Compositing/OpenColorIO/Source/OpenColorIOEditor/Private/OpenColorIOEditorSettings.cpp
--- a/Engine/Plugins/Compositing/OpenColorIO/Source/OpenColorIOEditor/Private/OpenColorIOEditorSettings.cpp
+++ b/Engine/Plugins/Compositing/OpenColorIO/Source/OpenColorIOEditor/Private/OpenColorIOEditorSettings.cpp
@@ -21,6 +21,14 @@ const FOpenColorIODisplayConfiguration* UOpenColorIOLevelViewportSettings::GetVi
 		return &Pair->DisplayConfiguration;
 	}
 
+	const UOpenColorIODefaultViewportSettings* DefaultViewportSettings = GetDefault<UOpenColorIODefaultViewportSettings>();
+	
+	if (DefaultViewportSettings->DefaultDisplayConfiguration.bIsEnabled)
+	{
+		// Apply default viewport settings if there is no locally cached viewport settings.
+		return &DefaultViewportSettings->DefaultDisplayConfiguration;
+	}
+
 	return nullptr;
 }
 
@@ -28,6 +36,8 @@ void UOpenColorIOLevelViewportSettings::PostInitProperties()
 {
 	Super::PostInitProperties();
 
+	// Note: This config could be moved back to (config = EditorPerProjectUserSettings) since level viewport layouts are local & user specific.
+
 	if (GConfig && FPaths::FileExists(GEditorPerProjectIni))
 	{
 		const TCHAR* SectionName = TEXT("/Script/OpenColorIOEditor.OpenColorIOLevelViewportSettings");
@@ -44,43 +54,69 @@ void UOpenColorIOLevelViewportSettings::PostInitProperties()
 		}
 	}
 
-	const bool bEnforceForwardViewDirectionOnly = !GetDefault<UOpenColorIOSettings>()->bSupportInverseViewTransforms;
-
-	for (FPerViewportDisplaySettingPair& ViewportSetting : ViewportsSettings)
-	{
-		if (bEnforceForwardViewDirectionOnly)
+	// Discard disabled empty/invalid settings
+	ViewportsSettings.RemoveAll([](const FPerViewportDisplaySettingPair& ViewportSetting)
 		{
-			ViewportSetting.DisplayConfiguration.ColorConfiguration.DisplayViewDirection = EOpenColorIOViewTransformDirection::Forward;
+			return !ViewportSetting.DisplayConfiguration.bIsEnabled && !ViewportSetting.DisplayConfiguration.ColorConfiguration.IsValid();
 		}
+	);
 
-		// Note: Ideally the following wouldn't be necessary but it is required since the previous default behavior was enabled with invalid settings.
-		if (!ViewportSetting.DisplayConfiguration.ColorConfiguration.IsValid())
+	const bool bEnforceForwardViewDirectionOnly = !GetDefault<UOpenColorIOSettings>()->bSupportInverseViewTransforms;
+	if (bEnforceForwardViewDirectionOnly)
+	{
+		for (FPerViewportDisplaySettingPair& ViewportSetting : ViewportsSettings)
 		{
-			UE_LOG(LogOpenColorIOEditor, Display, TEXT("Force-disable invalid viewport transform settings."));
-
-			ViewportSetting.DisplayConfiguration.bIsEnabled = false;
+			ViewportSetting.DisplayConfiguration.ColorConfiguration.DisplayViewDirection = EOpenColorIOViewTransformDirection::Forward;
 		}
 	}
 }
 
 void UOpenColorIOLevelViewportSettings::SetViewportSettings(FName ViewportIdentifier, const FOpenColorIODisplayConfiguration& Configuration)
 {
-	FPerViewportDisplaySettingPair* Pair = ViewportsSettings.FindByPredicate([ViewportIdentifier](const FPerViewportDisplaySettingPair& Other)
+	const UOpenColorIODefaultViewportSettings* DefaultViewportSettings = GetDefault<UOpenColorIODefaultViewportSettings>();
+
+	if (DefaultViewportSettings->DefaultDisplayConfiguration.Equals(Configuration))
+	{
+		// Current configuration matches the default, no need to update locally cached viewport settings.
+		return;
+	}
+
+	auto FindViewportFn = [ViewportIdentifier](const FPerViewportDisplaySettingPair& Other)
 		{
 			return Other.ViewportIdentifier == ViewportIdentifier;
-		});
+		};
 
-	if (Pair)
+	const bool bRemoveConfiguration = !Configuration.bIsEnabled && !Configuration.ColorConfiguration.IsValid();
+	if (bRemoveConfiguration)
 	{
-		Pair->DisplayConfiguration = Configuration;
+		ViewportsSettings.RemoveAll(FindViewportFn);
 	}
 	else
 	{
-		//Add new entry if viewport is not found
-		FPerViewportDisplaySettingPair NewEntry;
-		NewEntry.ViewportIdentifier = ViewportIdentifier;
-		NewEntry.DisplayConfiguration = Configuration;
-		ViewportsSettings.Emplace(MoveTemp(NewEntry));
+		FPerViewportDisplaySettingPair* Pair = ViewportsSettings.FindByPredicate(FindViewportFn);
+
+		if (Pair)
+		{
+			Pair->DisplayConfiguration = Configuration;
+		}
+		else
+		{
+			//Add new entry if viewport is not found
+			FPerViewportDisplaySettingPair NewEntry;
+			NewEntry.ViewportIdentifier = ViewportIdentifier;
+			NewEntry.DisplayConfiguration = Configuration;
+			ViewportsSettings.Emplace(MoveTemp(NewEntry));
+		}
 	}
 }
 
+void UOpenColorIODefaultViewportSettings::PostInitProperties()
+{
+	Super::PostInitProperties();
+
+	const bool bEnforceForwardViewDirectionOnly = !GetDefault<UOpenColorIOSettings>()->bSupportInverseViewTransforms;
+	if (bEnforceForwardViewDirectionOnly)
+	{
+		DefaultDisplayConfiguration.ColorConfiguration.DisplayViewDirection = EOpenColorIOViewTransformDirection::Forward;
+	}
+}
diff --git a/Engine/Plugins/Compositing/OpenColorIO/Source/OpenColorIOEditor/Private/OpenColorIOEditorSettings.h b/Engine/Plugins/Compositing/OpenColorIO/Source/OpenColorIOEditor/Private/OpenColorIOEditorSettings.h
--- a/Engine/Plugins/Compositing/OpenColorIO/Source/OpenColorIOEditor/Private/OpenColorIOEditorSettings.h
+++ b/Engine/Plugins/Compositing/OpenColorIO/Source/OpenColorIOEditor/Private/OpenColorIOEditorSettings.h
@@ -32,7 +32,6 @@ struct FPerViewportDisplaySettingPair
 /**
  * List of settings associated to level viewport instances linked with an identifier
  */
-
 UCLASS(config = OpenColorIO)
 class UOpenColorIOLevelViewportSettings : public UObject
 {
@@ -49,8 +48,22 @@ public:
 	void SetViewportSettings(FName ViewportIdentifier, const FOpenColorIODisplayConfiguration& Configuration);
 
 protected:
-	
+
 	/** Settings associated to each viewport that was configured */
 	UPROPERTY(config)
 	TArray<FPerViewportDisplaySettingPair> ViewportsSettings;
 };
+
+UCLASS(config = OpenColorIO)
+class UOpenColorIODefaultViewportSettings : public UObject
+{
+	GENERATED_BODY()
+
+public:
+
+	virtual void PostInitProperties() override;
+
+	/** Default settings applied to level viewports when none was previously set, or if settings have been cleared. */
+	UPROPERTY(config, EditAnywhere, Category = OpenColorIO, DisplayName="Default Viewport Configuration")
+	FOpenColorIODisplayConfiguration DefaultDisplayConfiguration;
+};
