OCIO: Add support for color & display-view tooltip descriptions.

diff --git a/Engine/Plugins/Compositing/OpenColorIO/Source/OpenColorIOEditor/Private/OpenColorIOColorSpaceCustomization.cpp b/Engine/Plugins/Compositing/OpenColorIO/Source/OpenColorIOEditor/Private/OpenColorIOColorSpaceCustomization.cpp
--- a/Engine/Plugins/Compositing/OpenColorIO/Source/OpenColorIOEditor/Private/OpenColorIOColorSpaceCustomization.cpp
+++ b/Engine/Plugins/Compositing/OpenColorIO/Source/OpenColorIOEditor/Private/OpenColorIOColorSpaceCustomization.cpp
@@ -157,7 +157,7 @@ FOpenColorIOWrapperConfig* IPropertyTypeCustomizationOpenColorIO::GetConfigWrapp
 	return nullptr;
 }
 
-void FOpenColorIOColorSpaceCustomization::ProcessColorSpaceForMenuGeneration(FMenuBuilder& InMenuBuilder, const int32 InMenuDepth, const FString& InPreviousFamilyHierarchy, const FOpenColorIOColorSpace& InColorSpace, TArray<FString>& InOutExistingMenuFilter)
+void FOpenColorIOColorSpaceCustomization::ProcessColorSpaceForMenuGeneration(FMenuBuilder& InMenuBuilder, const int32 InMenuDepth, const FString& InPreviousFamilyHierarchy, const FOpenColorIOColorSpace& InColorSpace, TArray<FString>& InOutExistingMenuFilter, const FString& InColorSpaceDescription)
 {
 	const FString NextFamilyName = InColorSpace.GetFamilyNameAtDepth(InMenuDepth);
 	if (!NextFamilyName.IsEmpty())
@@ -179,8 +179,7 @@ void FOpenColorIOColorSpaceCustomization::ProcessColorSpaceForMenuGeneration(FMe
 	}
 	else
 	{
-
-		AddMenuEntry(InMenuBuilder, InColorSpace);
+		AddMenuEntry(InMenuBuilder, InColorSpace, InColorSpaceDescription);
 	}
 }
 
@@ -201,22 +200,23 @@ void FOpenColorIOColorSpaceCustomization::PopulateSubMenu(FMenuBuilder& InMenuBu
 			FOpenColorIOColorSpace ColorSpace;
 			ColorSpace.ColorSpaceName = ConfigWrapper->GetColorSpaceName(i);
 			ColorSpace.FamilyName = ConfigWrapper->GetColorSpaceFamilyName(*ColorSpace.ColorSpaceName);
+			const FString ColorSpaceDesc = ConfigWrapper->GetColorSpaceDescription(*ColorSpace.ColorSpaceName);
 
 			//Filter out color spaces that don't belong to this hierarchy
 			if (InPreviousFamilyHierarchy.IsEmpty() || ColorSpace.FamilyName.Contains(InPreviousFamilyHierarchy))
 			{
-				ProcessColorSpaceForMenuGeneration(InMenuBuilder, InMenuDepth, InPreviousFamilyHierarchy, ColorSpace, ExistingSubMenus);
+				ProcessColorSpaceForMenuGeneration(InMenuBuilder, InMenuDepth, InPreviousFamilyHierarchy, ColorSpace, ExistingSubMenus, ColorSpaceDesc);
 			}
 		}
 	}
 
 }
 
-void FOpenColorIOColorSpaceCustomization::AddMenuEntry(FMenuBuilder& InMenuBuilder, const FOpenColorIOColorSpace& InColorSpace)
+void FOpenColorIOColorSpaceCustomization::AddMenuEntry(FMenuBuilder& InMenuBuilder, const FOpenColorIOColorSpace& InColorSpace, const FString& InDescription)
 {
 	InMenuBuilder.AddMenuEntry(
 		FText::FromString(InColorSpace.ToString()),
-		FText::FromString(InColorSpace.ToString()),
+		FText::FromString(InDescription),
 		FSlateIcon(),
 		FUIAction(
 			FExecuteAction::CreateLambda([this, InColorSpace]
@@ -266,11 +266,12 @@ TSharedRef<SWidget> FOpenColorIOColorSpaceCustomization::HandleSourceComboButton
 				FOpenColorIOColorSpace ColorSpace;
 				ColorSpace.ColorSpaceName = ConfigWrapper->GetColorSpaceName(i);
 				ColorSpace.FamilyName = ConfigWrapper->GetColorSpaceFamilyName(*ColorSpace.ColorSpaceName);
+				const FString ColorSpaceDesc = ConfigWrapper->GetColorSpaceDescription(*ColorSpace.ColorSpaceName);
 
 				//Top level menus have no preceding hierarchy.
 				const int32 CurrentMenuDepth = 0;
 				const FString PreviousHierarchy = FString();
-				ProcessColorSpaceForMenuGeneration(MenuBuilder, CurrentMenuDepth, PreviousHierarchy, ColorSpace, ExistingSubMenus);
+				ProcessColorSpaceForMenuGeneration(MenuBuilder, CurrentMenuDepth, PreviousHierarchy, ColorSpace, ExistingSubMenus, ColorSpaceDesc);
 			}
 
 			if (ColorSpaceCount <= 0)
@@ -297,17 +298,18 @@ void FOpenColorIODisplayViewCustomization::PopulateViewSubMenu(FMenuBuilder& InM
 		for (int32 i = 0; i < ConfigWrapper->GetNumViews(*InDisplayView.Display); ++i)
 		{
 			InDisplayView.View = ConfigWrapper->GetViewName(*InDisplayView.Display, i);
+			const FString Description = ConfigWrapper->GetDisplayViewTransformDescription(*InDisplayView.Display, *InDisplayView.View);
 
-			AddMenuEntry(InMenuBuilder, InDisplayView);
+			AddMenuEntry(InMenuBuilder, InDisplayView, Description);
 		}
 	}
 }
 
-void FOpenColorIODisplayViewCustomization::AddMenuEntry(FMenuBuilder& InMenuBuilder, const FOpenColorIODisplayView& InDisplayView)
+void FOpenColorIODisplayViewCustomization::AddMenuEntry(FMenuBuilder& InMenuBuilder, const FOpenColorIODisplayView& InDisplayView, const FString& InDescription)
 {
 	InMenuBuilder.AddMenuEntry(
 		FText::FromString(InDisplayView.View),
-		FText::FromString(InDisplayView.ToString()),
+		FText::FromString(InDescription),
 		FSlateIcon(),
 		FUIAction(
 			FExecuteAction::CreateLambda([this, InDisplayView]
diff --git a/Engine/Plugins/Compositing/OpenColorIO/Source/OpenColorIOEditor/Private/OpenColorIOColorSpaceCustomization.h b/Engine/Plugins/Compositing/OpenColorIO/Source/OpenColorIOEditor/Private/OpenColorIOColorSpaceCustomization.h
--- a/Engine/Plugins/Compositing/OpenColorIO/Source/OpenColorIOEditor/Private/OpenColorIOColorSpaceCustomization.h
+++ b/Engine/Plugins/Compositing/OpenColorIO/Source/OpenColorIOEditor/Private/OpenColorIOColorSpaceCustomization.h
@@ -50,9 +50,9 @@ public:
 	virtual void CustomizeHeader(TSharedRef<class IPropertyHandle> InPropertyHandle, class FDetailWidgetRow& HeaderRow, IPropertyTypeCustomizationUtils& PropertyTypeCustomizationUtils) override;
 
 private:
-	void ProcessColorSpaceForMenuGeneration(FMenuBuilder& InMenuBuilder, const int32 InMenuDepth, const FString& InPreviousFamilyHierarchy, const FOpenColorIOColorSpace& InColorSpace, TArray<FString>& InOutExistingMenuFilter);
+	void ProcessColorSpaceForMenuGeneration(FMenuBuilder& InMenuBuilder, const int32 InMenuDepth, const FString& InPreviousFamilyHierarchy, const FOpenColorIOColorSpace& InColorSpace, TArray<FString>& InOutExistingMenuFilter, const FString& InColorSpaceDescription);
 	void PopulateSubMenu(FMenuBuilder& InMenuBuilder, const int32 InMenuDepth, FString InPreviousFamilyHierarchy);
-	void AddMenuEntry(FMenuBuilder& InMenuBuilder, const FOpenColorIOColorSpace& InColorSpace);
+	void AddMenuEntry(FMenuBuilder& InMenuBuilder, const FOpenColorIOColorSpace& InColorSpace, const FString& InDescription);
 
 	TSharedRef<SWidget> HandleSourceComboButtonMenuContent();
 };
@@ -72,7 +72,7 @@ public:
 
 private:
 	void PopulateViewSubMenu(FMenuBuilder& InMenuBuilder, FOpenColorIODisplayView InDisplayView);
-	void AddMenuEntry(FMenuBuilder& InMenuBuilder, const FOpenColorIODisplayView& InDisplayView);
+	void AddMenuEntry(FMenuBuilder& InMenuBuilder, const FOpenColorIODisplayView& InDisplayView, const FString& InDescription);
 
 	TSharedRef<SWidget> HandleSourceComboButtonMenuContent();
 };
diff --git a/Engine/Source/Runtime/OpenColorIOWrapper/Private/OpenColorIOWrapper.cpp b/Engine/Source/Runtime/OpenColorIOWrapper/Private/OpenColorIOWrapper.cpp
--- a/Engine/Source/Runtime/OpenColorIOWrapper/Private/OpenColorIOWrapper.cpp
+++ b/Engine/Source/Runtime/OpenColorIOWrapper/Private/OpenColorIOWrapper.cpp
@@ -174,6 +174,7 @@ FOpenColorIOWrapperConfig::FOpenColorIOWrapperConfig(FStringView InFilePath, EAd
 					ColorSpaceRcPtr WorkingCS = InterchangeCS->createEditableCopy();
 					WorkingCS->setName(AnsiWorkingColorSpaceName.Get());
 					WorkingCS->setFamily("UE");
+					WorkingCS->setDescription("The global engine working color space, as defined in project settings (defaults to sRGB/Rec.709).");
 					WorkingCS->clearAliases();
 
 					ConfigRcPtr NewConfigCopy = NewConfig->createEditableCopy();
@@ -245,6 +246,21 @@ FString FOpenColorIOWrapperConfig::GetColorSpaceFamilyName(const TCHAR* InColorS
 	return {};
 }
 
+FString FOpenColorIOWrapperConfig::GetColorSpaceDescription(const TCHAR* InColorSpaceName) const
+{
+	if (IsValid())
+	{
+		OCIO_NAMESPACE::ConstColorSpaceRcPtr ColorSpace = Pimpl->Config->getColorSpace(StringCast<ANSICHAR>(InColorSpaceName).Get());
+
+		if (ColorSpace != nullptr)
+		{
+			return StringCast<TCHAR>(ColorSpace->getDescription()).Get();
+		}
+	}
+
+	return {};
+}
+
 int32 FOpenColorIOWrapperConfig::GetNumDisplays() const
 {
 	if (IsValid())
@@ -301,6 +317,25 @@ FString FOpenColorIOWrapperConfig::GetDisplayViewTransformName(const TCHAR* InDi
 	return {};
 }
 
+FString FOpenColorIOWrapperConfig::GetDisplayViewTransformDescription(const TCHAR* InDisplayName, const TCHAR* InViewName) const
+{
+	if (IsValid())
+	{
+		const char* TransformName = Pimpl->Config->getDisplayViewTransformName(StringCast<ANSICHAR>(InDisplayName).Get(), StringCast<ANSICHAR>(InViewName).Get());
+
+		OCIO_NAMESPACE::ConstViewTransformRcPtr ViewTransform = Pimpl->Config->getViewTransform(TransformName);
+		
+		if (ViewTransform != nullptr)
+		{
+			const char* TransformDescription = ViewTransform->getDescription();
+
+			return StringCast<TCHAR>(TransformDescription).Get();
+		}
+	}
+
+	return {};
+}
+
 TMap<FString, FString> FOpenColorIOWrapperConfig::GetCurrentContextStringVars() const
 {
 	OCIO_EXCEPTION_HANDLING_TRY();
diff --git a/Engine/Source/Runtime/OpenColorIOWrapper/Public/OpenColorIOWrapper.h b/Engine/Source/Runtime/OpenColorIOWrapper/Public/OpenColorIOWrapper.h
--- a/Engine/Source/Runtime/OpenColorIOWrapper/Public/OpenColorIOWrapper.h
+++ b/Engine/Source/Runtime/OpenColorIOWrapper/Public/OpenColorIOWrapper.h
@@ -86,6 +86,9 @@ public:
 	
 	/** Get the family name for a color space. */
 	OPENCOLORIOWRAPPER_API FString GetColorSpaceFamilyName(const TCHAR* InColorSpaceName) const;
+
+	/** Get a color space description. */
+	OPENCOLORIOWRAPPER_API FString GetColorSpaceDescription(const TCHAR* InColorSpaceName) const;
 	
 	/** Get the number of displays in the configuration. */
 	OPENCOLORIOWRAPPER_API int32 GetNumDisplays() const;
@@ -99,9 +102,12 @@ public:
 	/** Get a view name for its display and index. */
 	OPENCOLORIOWRAPPER_API FString GetViewName(const TCHAR* InDisplayName, int32 Index) const;
 
-	/** Get a display-view trnasform name. */
+	/** Get a display-view transform name. */
 	OPENCOLORIOWRAPPER_API FString GetDisplayViewTransformName(const TCHAR* InDisplayName, const TCHAR* InViewName) const;
 
+	/** Get a display-view transform description. */
+	OPENCOLORIOWRAPPER_API FString GetDisplayViewTransformDescription(const TCHAR* InDisplayName, const TCHAR* InViewName) const;
+
 	/** Returns the current context key-value strings. */
 	OPENCOLORIOWRAPPER_API TMap<FString, FString> GetCurrentContextStringVars() const;
 
