OCIO: Small fix to let users disable an invalid OCIO transform.

The disabled invalid transform will now be correctly removed from saved settings on-close or viewport change.

diff --git a/Engine/Plugins/Compositing/OpenColorIO/Source/OpenColorIOEditor/Private/OpenColorIOEditorSettings.cpp b/Engine/Plugins/Compositing/OpenColorIO/Source/OpenColorIOEditor/Private/OpenColorIOEditorSettings.cpp
--- a/Engine/Plugins/Compositing/OpenColorIO/Source/OpenColorIOEditor/Private/OpenColorIOEditorSettings.cpp
+++ b/Engine/Plugins/Compositing/OpenColorIO/Source/OpenColorIOEditor/Private/OpenColorIOEditorSettings.cpp
@@ -73,40 +73,30 @@ void UOpenColorIOLevelViewportSettings::PostInitProperties()
 
 void UOpenColorIOLevelViewportSettings::SetViewportSettings(FName ViewportIdentifier, const FOpenColorIODisplayConfiguration& Configuration)
 {
-	const UOpenColorIODefaultViewportSettings* DefaultViewportSettings = GetDefault<UOpenColorIODefaultViewportSettings>();
-
-	if (DefaultViewportSettings->DefaultDisplayConfiguration.Equals(Configuration))
-	{
-		// Current configuration matches the default, no need to update locally cached viewport settings.
-		return;
-	}
-
-	auto FindViewportFn = [ViewportIdentifier](const FPerViewportDisplaySettingPair& Other)
+	auto FindMatchingViewportFn = [ViewportIdentifier](const FPerViewportDisplaySettingPair& Other)
 		{
 			return Other.ViewportIdentifier == ViewportIdentifier;
 		};
 
 	const bool bRemoveConfiguration = !Configuration.bIsEnabled && !Configuration.ColorConfiguration.IsValid();
-	if (bRemoveConfiguration)
+	const UOpenColorIODefaultViewportSettings* DefaultViewportSettings = GetDefault<UOpenColorIODefaultViewportSettings>();
+
+	if (bRemoveConfiguration || DefaultViewportSettings->DefaultDisplayConfiguration.Equals(Configuration))
 	{
-		ViewportsSettings.RemoveAll(FindViewportFn);
+		// Remove settings if they are disabled & invalid, or if they match the default viewport configuration
+		ViewportsSettings.RemoveAll(FindMatchingViewportFn);
+	}
+	else if (FPerViewportDisplaySettingPair* Pair = ViewportsSettings.FindByPredicate(FindMatchingViewportFn))
+	{
+		Pair->DisplayConfiguration = Configuration;
 	}
 	else
 	{
-		FPerViewportDisplaySettingPair* Pair = ViewportsSettings.FindByPredicate(FindViewportFn);
-
-		if (Pair)
-		{
-			Pair->DisplayConfiguration = Configuration;
-		}
-		else
-		{
-			//Add new entry if viewport is not found
-			FPerViewportDisplaySettingPair NewEntry;
-			NewEntry.ViewportIdentifier = ViewportIdentifier;
-			NewEntry.DisplayConfiguration = Configuration;
-			ViewportsSettings.Emplace(MoveTemp(NewEntry));
-		}
+		//Add new entry if viewport is not found
+		FPerViewportDisplaySettingPair NewEntry;
+		NewEntry.ViewportIdentifier = ViewportIdentifier;
+		NewEntry.DisplayConfiguration = Configuration;
+		ViewportsSettings.Emplace(MoveTemp(NewEntry));
 	}
 }
 
