OCIO: Exposed description on all colorspace or display-view transform tooltip selections.

Descriptions will be accessible only on newly created (or force-reloaded) configuration asset transforms, since we don't want to invalidate existing content.

diff --git a/Engine/Plugins/Compositing/OpenColorIO/Source/OpenColorIO/Private/OpenColorIOColorSpace.cpp b/Engine/Plugins/Compositing/OpenColorIO/Source/OpenColorIO/Private/OpenColorIOColorSpace.cpp
--- a/Engine/Plugins/Compositing/OpenColorIO/Source/OpenColorIO/Private/OpenColorIOColorSpace.cpp
+++ b/Engine/Plugins/Compositing/OpenColorIO/Source/OpenColorIO/Private/OpenColorIOColorSpace.cpp
@@ -14,9 +14,10 @@
 
 const TCHAR* FOpenColorIOColorSpace::FamilyDelimiter = TEXT("/");
 
-FOpenColorIOColorSpace::FOpenColorIOColorSpace(const FString& InColorSpaceName, int32 InColorSpaceIndex, const FString& InFamilyName)
+FOpenColorIOColorSpace::FOpenColorIOColorSpace(const FString& InColorSpaceName, int32 InColorSpaceIndex, const FString& InFamilyName, const FString& InDescription)
 	: ColorSpaceName(InColorSpaceName)
 	, FamilyName(InFamilyName)
+	, Description(InDescription)
 { }
 
 FString FOpenColorIOColorSpace::ToString() const
@@ -37,6 +38,7 @@ void FOpenColorIOColorSpace::Reset()
 {
 	ColorSpaceName.Reset();
 	FamilyName.Reset();
+	Description.Reset();
 }
 
 FString FOpenColorIOColorSpace::GetFamilyNameAtDepth(int32 InDepth) const
@@ -65,15 +67,10 @@ FString FOpenColorIOColorSpace::GetFamilyNameAtDepth(int32 InDepth) const
  * FOpenColorIODisplayView implementation
  */
 
-FOpenColorIODisplayView::FOpenColorIODisplayView()
-	: Display()
-	, View()
-{
-}
-
-FOpenColorIODisplayView::FOpenColorIODisplayView(FStringView InDisplayName, FStringView InViewName)
+FOpenColorIODisplayView::FOpenColorIODisplayView(FStringView InDisplayName, FStringView InViewName, const FString& InDescription)
 	: Display(InDisplayName)
 	, View(InViewName)
+	, Description(InDescription)
 {
 }
 
@@ -96,6 +93,7 @@ void FOpenColorIODisplayView::Reset()
 {
 	Display.Reset();
 	View.Reset();
+	Description.Reset();
 }
 
 /*
diff --git a/Engine/Plugins/Compositing/OpenColorIO/Source/OpenColorIO/Private/OpenColorIOConfiguration.cpp b/Engine/Plugins/Compositing/OpenColorIO/Source/OpenColorIO/Private/OpenColorIOConfiguration.cpp
--- a/Engine/Plugins/Compositing/OpenColorIO/Source/OpenColorIO/Private/OpenColorIOConfiguration.cpp
+++ b/Engine/Plugins/Compositing/OpenColorIO/Source/OpenColorIO/Private/OpenColorIOConfiguration.cpp
@@ -222,6 +222,8 @@ void UOpenColorIOConfiguration::ReloadExistingColorspaces(bool bForce)
 				FOpenColorIOColorSpace ColorSpace;
 				ColorSpace.ColorSpaceName = ExistingColorSpace.ColorSpaceName;
 				ColorSpace.FamilyName = Config->GetColorSpaceFamilyName(*ExistingColorSpace.ColorSpaceName);
+				ColorSpace.Description = Config->GetColorSpaceDescription(*ExistingColorSpace.ColorSpaceName);
+
 				DesiredColorSpaces.Add(ColorSpace);
 			}
 
@@ -234,7 +236,10 @@ void UOpenColorIOConfiguration::ReloadExistingColorspaces(bool bForce)
 
 					if (ViewName == ExistingDisplayView.View)
 					{
-						DesiredDisplayViews.Add(ExistingDisplayView);
+						FOpenColorIODisplayView DisplayView = ExistingDisplayView;
+						DisplayView.Description = Config->GetDisplayViewTransformDescription(*ExistingDisplayView.Display, *ExistingDisplayView.View);
+						
+						DesiredDisplayViews.Add(DisplayView);
 						bDisplayViewFound = true;
 						break;
 					}
diff --git a/Engine/Plugins/Compositing/OpenColorIO/Source/OpenColorIO/Public/OpenColorIOColorSpace.h b/Engine/Plugins/Compositing/OpenColorIO/Source/OpenColorIO/Public/OpenColorIOColorSpace.h
--- a/Engine/Plugins/Compositing/OpenColorIO/Source/OpenColorIO/Public/OpenColorIOColorSpace.h
+++ b/Engine/Plugins/Compositing/OpenColorIO/Source/OpenColorIO/Public/OpenColorIOColorSpace.h
@@ -35,7 +35,7 @@ public:
 	/**
 	 * Create and initialize a new instance.
 	 */
-	FOpenColorIOColorSpace(const FString& InColorSpaceName, int32 InColorSpaceIndex, const FString& InFamilyName);
+	FOpenColorIOColorSpace(const FString& InColorSpaceName, int32 InColorSpaceIndex, const FString& InFamilyName, const FString& InDescription = FString());
 
 	/** The ColorSpace name. */
 	UPROPERTY(EditAnywhere, BlueprintReadWrite, Category=ColorSpace)
@@ -55,6 +55,10 @@ public:
 	UPROPERTY(EditAnywhere, BlueprintReadWrite, Category=ColorSpace)
 	FString FamilyName;
 
+	/** Colorspace description. */
+	UPROPERTY()
+	FString Description;
+
 	/** Delimiter used in the OpenColorIO library to make family hierarchies */
 	static const TCHAR* FamilyDelimiter;
 
@@ -101,19 +105,25 @@ struct OPENCOLORIO_API FOpenColorIODisplayView
 
 public:
 	/** Default constructor. */
-	FOpenColorIODisplayView();
+	FOpenColorIODisplayView() = default;
 
 	/**
 	 * Create and initialize a new instance.
 	 */
-	FOpenColorIODisplayView(FStringView InDisplayName, FStringView InViewName);
+	FOpenColorIODisplayView(FStringView InDisplayName, FStringView InViewName, const FString& InDescription = FString());
 
+	/** Display name. */
 	UPROPERTY(EditAnywhere, BlueprintReadWrite, Category = ColorSpace)
 	FString Display;
 
+	/** View name. */
 	UPROPERTY(EditAnywhere, BlueprintReadWrite, Category = ColorSpace)
 	FString View;
 
+	/** Display-view (transform) description. */
+	UPROPERTY()
+	FString Description;
+
 	FString ToString() const;
 
 	/** Return true if the index and name have been set properly */
diff --git a/Engine/Plugins/Compositing/OpenColorIO/Source/OpenColorIOEditor/Private/OpenColorIOColorSpaceCustomization.cpp b/Engine/Plugins/Compositing/OpenColorIO/Source/OpenColorIOEditor/Private/OpenColorIOColorSpaceCustomization.cpp
--- a/Engine/Plugins/Compositing/OpenColorIO/Source/OpenColorIOEditor/Private/OpenColorIOColorSpaceCustomization.cpp
+++ b/Engine/Plugins/Compositing/OpenColorIO/Source/OpenColorIOEditor/Private/OpenColorIOColorSpaceCustomization.cpp
@@ -157,7 +157,7 @@ FOpenColorIOWrapperConfig* IPropertyTypeCustomizationOpenColorIO::GetConfigWrapp
 	return nullptr;
 }
 
-void FOpenColorIOColorSpaceCustomization::ProcessColorSpaceForMenuGeneration(FMenuBuilder& InMenuBuilder, const int32 InMenuDepth, const FString& InPreviousFamilyHierarchy, const FOpenColorIOColorSpace& InColorSpace, TArray<FString>& InOutExistingMenuFilter, const FString& InColorSpaceDescription)
+void FOpenColorIOColorSpaceCustomization::ProcessColorSpaceForMenuGeneration(FMenuBuilder& InMenuBuilder, const int32 InMenuDepth, const FString& InPreviousFamilyHierarchy, const FOpenColorIOColorSpace& InColorSpace, TArray<FString>& InOutExistingMenuFilter)
 {
 	const FString NextFamilyName = InColorSpace.GetFamilyNameAtDepth(InMenuDepth);
 	if (!NextFamilyName.IsEmpty())
@@ -179,7 +179,7 @@ void FOpenColorIOColorSpaceCustomization::ProcessColorSpaceForMenuGeneration(FMe
 	}
 	else
 	{
-		AddMenuEntry(InMenuBuilder, InColorSpace, InColorSpaceDescription);
+		AddMenuEntry(InMenuBuilder, InColorSpace);
 	}
 }
 
@@ -200,23 +200,23 @@ void FOpenColorIOColorSpaceCustomization::PopulateSubMenu(FMenuBuilder& InMenuBu
 			FOpenColorIOColorSpace ColorSpace;
 			ColorSpace.ColorSpaceName = ConfigWrapper->GetColorSpaceName(i);
 			ColorSpace.FamilyName = ConfigWrapper->GetColorSpaceFamilyName(*ColorSpace.ColorSpaceName);
-			const FString ColorSpaceDesc = ConfigWrapper->GetColorSpaceDescription(*ColorSpace.ColorSpaceName);
+			ColorSpace.Description	= ConfigWrapper->GetColorSpaceDescription(*ColorSpace.ColorSpaceName);
 
 			//Filter out color spaces that don't belong to this hierarchy
 			if (InPreviousFamilyHierarchy.IsEmpty() || ColorSpace.FamilyName.Contains(InPreviousFamilyHierarchy))
 			{
-				ProcessColorSpaceForMenuGeneration(InMenuBuilder, InMenuDepth, InPreviousFamilyHierarchy, ColorSpace, ExistingSubMenus, ColorSpaceDesc);
+				ProcessColorSpaceForMenuGeneration(InMenuBuilder, InMenuDepth, InPreviousFamilyHierarchy, ColorSpace, ExistingSubMenus);
 			}
 		}
 	}
 
 }
 
-void FOpenColorIOColorSpaceCustomization::AddMenuEntry(FMenuBuilder& InMenuBuilder, const FOpenColorIOColorSpace& InColorSpace, const FString& InDescription)
+void FOpenColorIOColorSpaceCustomization::AddMenuEntry(FMenuBuilder& InMenuBuilder, const FOpenColorIOColorSpace& InColorSpace)
 {
 	InMenuBuilder.AddMenuEntry(
 		FText::FromString(InColorSpace.ToString()),
-		FText::FromString(InDescription),
+		FText::FromString(InColorSpace.Description),
 		FSlateIcon(),
 		FUIAction(
 			FExecuteAction::CreateLambda([this, InColorSpace]
@@ -266,12 +266,12 @@ TSharedRef<SWidget> FOpenColorIOColorSpaceCustomization::HandleSourceComboButton
 				FOpenColorIOColorSpace ColorSpace;
 				ColorSpace.ColorSpaceName = ConfigWrapper->GetColorSpaceName(i);
 				ColorSpace.FamilyName = ConfigWrapper->GetColorSpaceFamilyName(*ColorSpace.ColorSpaceName);
-				const FString ColorSpaceDesc = ConfigWrapper->GetColorSpaceDescription(*ColorSpace.ColorSpaceName);
+				ColorSpace.Description = ConfigWrapper->GetColorSpaceDescription(*ColorSpace.ColorSpaceName);
 
 				//Top level menus have no preceding hierarchy.
 				const int32 CurrentMenuDepth = 0;
 				const FString PreviousHierarchy = FString();
-				ProcessColorSpaceForMenuGeneration(MenuBuilder, CurrentMenuDepth, PreviousHierarchy, ColorSpace, ExistingSubMenus, ColorSpaceDesc);
+				ProcessColorSpaceForMenuGeneration(MenuBuilder, CurrentMenuDepth, PreviousHierarchy, ColorSpace, ExistingSubMenus);
 			}
 
 			if (ColorSpaceCount <= 0)
@@ -298,18 +298,18 @@ void FOpenColorIODisplayViewCustomization::PopulateViewSubMenu(FMenuBuilder& InM
 		for (int32 i = 0; i < ConfigWrapper->GetNumViews(*InDisplayView.Display); ++i)
 		{
 			InDisplayView.View = ConfigWrapper->GetViewName(*InDisplayView.Display, i);
-			const FString Description = ConfigWrapper->GetDisplayViewTransformDescription(*InDisplayView.Display, *InDisplayView.View);
+			InDisplayView.Description = ConfigWrapper->GetDisplayViewTransformDescription(*InDisplayView.Display, *InDisplayView.View);
 
-			AddMenuEntry(InMenuBuilder, InDisplayView, Description);
+			AddMenuEntry(InMenuBuilder, InDisplayView);
 		}
 	}
 }
 
-void FOpenColorIODisplayViewCustomization::AddMenuEntry(FMenuBuilder& InMenuBuilder, const FOpenColorIODisplayView& InDisplayView, const FString& InDescription)
+void FOpenColorIODisplayViewCustomization::AddMenuEntry(FMenuBuilder& InMenuBuilder, const FOpenColorIODisplayView& InDisplayView)
 {
 	InMenuBuilder.AddMenuEntry(
 		FText::FromString(InDisplayView.View),
-		FText::FromString(InDescription),
+		FText::FromString(InDisplayView.Description),
 		FSlateIcon(),
 		FUIAction(
 			FExecuteAction::CreateLambda([this, InDisplayView]
diff --git a/Engine/Plugins/Compositing/OpenColorIO/Source/OpenColorIOEditor/Private/OpenColorIOColorSpaceCustomization.h b/Engine/Plugins/Compositing/OpenColorIO/Source/OpenColorIOEditor/Private/OpenColorIOColorSpaceCustomization.h
--- a/Engine/Plugins/Compositing/OpenColorIO/Source/OpenColorIOEditor/Private/OpenColorIOColorSpaceCustomization.h
+++ b/Engine/Plugins/Compositing/OpenColorIO/Source/OpenColorIOEditor/Private/OpenColorIOColorSpaceCustomization.h
@@ -50,9 +50,9 @@ public:
 	virtual void CustomizeHeader(TSharedRef<class IPropertyHandle> InPropertyHandle, class FDetailWidgetRow& HeaderRow, IPropertyTypeCustomizationUtils& PropertyTypeCustomizationUtils) override;
 
 private:
-	void ProcessColorSpaceForMenuGeneration(FMenuBuilder& InMenuBuilder, const int32 InMenuDepth, const FString& InPreviousFamilyHierarchy, const FOpenColorIOColorSpace& InColorSpace, TArray<FString>& InOutExistingMenuFilter, const FString& InColorSpaceDescription);
+	void ProcessColorSpaceForMenuGeneration(FMenuBuilder& InMenuBuilder, const int32 InMenuDepth, const FString& InPreviousFamilyHierarchy, const FOpenColorIOColorSpace& InColorSpace, TArray<FString>& InOutExistingMenuFilter);
 	void PopulateSubMenu(FMenuBuilder& InMenuBuilder, const int32 InMenuDepth, FString InPreviousFamilyHierarchy);
-	void AddMenuEntry(FMenuBuilder& InMenuBuilder, const FOpenColorIOColorSpace& InColorSpace, const FString& InDescription);
+	void AddMenuEntry(FMenuBuilder& InMenuBuilder, const FOpenColorIOColorSpace& InColorSpace);
 
 	TSharedRef<SWidget> HandleSourceComboButtonMenuContent();
 };
@@ -72,7 +72,7 @@ public:
 
 private:
 	void PopulateViewSubMenu(FMenuBuilder& InMenuBuilder, FOpenColorIODisplayView InDisplayView);
-	void AddMenuEntry(FMenuBuilder& InMenuBuilder, const FOpenColorIODisplayView& InDisplayView, const FString& InDescription);
+	void AddMenuEntry(FMenuBuilder& InMenuBuilder, const FOpenColorIODisplayView& InDisplayView);
 
 	TSharedRef<SWidget> HandleSourceComboButtonMenuContent();
 };
diff --git a/Engine/Plugins/Compositing/OpenColorIO/Source/OpenColorIOEditor/Private/Widgets/SOpenColorIOColorSpacePicker.cpp b/Engine/Plugins/Compositing/OpenColorIO/Source/OpenColorIOEditor/Private/Widgets/SOpenColorIOColorSpacePicker.cpp
--- a/Engine/Plugins/Compositing/OpenColorIO/Source/OpenColorIOEditor/Private/Widgets/SOpenColorIOColorSpacePicker.cpp
+++ b/Engine/Plugins/Compositing/OpenColorIO/Source/OpenColorIOEditor/Private/Widgets/SOpenColorIOColorSpacePicker.cpp
@@ -129,7 +129,7 @@ TSharedRef<SWidget> SOpenColorIOColorSpacePicker::HandleColorSpaceComboButtonMen
 				MenuBuilder.AddMenuEntry
 				(
 					FText::FromString(ColorSpaceStr),
-					FText::FromString(ColorSpaceStr),
+					FText::FromString(ColorSpace.Description),
 					FSlateIcon(),
 					FUIAction
 					(
@@ -180,7 +180,7 @@ TSharedRef<SWidget> SOpenColorIOColorSpacePicker::HandleColorSpaceComboButtonMen
 					MenuBuilder.AddMenuEntry
 					(
 						FText::FromString(DisplayViewStr),
-						FText::FromString(DisplayViewStr),
+						FText::FromString(DisplayView.Description),
 						FSlateIcon(),
 						FUIAction
 						(
diff --git a/Engine/Source/Runtime/OpenColorIOWrapper/Private/OpenColorIOWrapper.cpp b/Engine/Source/Runtime/OpenColorIOWrapper/Private/OpenColorIOWrapper.cpp
--- a/Engine/Source/Runtime/OpenColorIOWrapper/Private/OpenColorIOWrapper.cpp
+++ b/Engine/Source/Runtime/OpenColorIOWrapper/Private/OpenColorIOWrapper.cpp
@@ -317,6 +317,18 @@ FString FOpenColorIOWrapperConfig::GetDisplayViewTransformName(const TCHAR* InDi
 	return {};
 }
 
+FString FOpenColorIOWrapperConfig::GetDisplayViewDescription(const TCHAR* InDisplayName, const TCHAR* InViewName) const
+{
+	if (IsValid())
+	{
+		const char* DisplayViewDescription = Pimpl->Config->getDisplayViewDescription(StringCast<ANSICHAR>(InDisplayName).Get(), StringCast<ANSICHAR>(InViewName).Get());
+
+		return StringCast<TCHAR>(DisplayViewDescription).Get();
+	}
+
+	return {};
+}
+
 FString FOpenColorIOWrapperConfig::GetDisplayViewTransformDescription(const TCHAR* InDisplayName, const TCHAR* InViewName) const
 {
 	if (IsValid())
diff --git a/Engine/Source/Runtime/OpenColorIOWrapper/Public/OpenColorIOWrapper.h b/Engine/Source/Runtime/OpenColorIOWrapper/Public/OpenColorIOWrapper.h
--- a/Engine/Source/Runtime/OpenColorIOWrapper/Public/OpenColorIOWrapper.h
+++ b/Engine/Source/Runtime/OpenColorIOWrapper/Public/OpenColorIOWrapper.h
@@ -105,6 +105,9 @@ public:
 	/** Get a display-view transform name. */
 	OPENCOLORIOWRAPPER_API FString GetDisplayViewTransformName(const TCHAR* InDisplayName, const TCHAR* InViewName) const;
 
+	/** Get a display-view transform name. */
+	OPENCOLORIOWRAPPER_API FString GetDisplayViewDescription(const TCHAR* InDisplayName, const TCHAR* InViewName) const;
+
 	/** Get a display-view transform description. */
 	OPENCOLORIOWRAPPER_API FString GetDisplayViewTransformDescription(const TCHAR* InDisplayName, const TCHAR* InViewName) const;
 
