OCIO: Fix for transforms failing to apply in shipping games.

Shader parameter metadata is always needed.

diff --git a/Engine/Plugins/Compositing/OpenColorIO/Source/OpenColorIO/Private/OpenColorIOShared.cpp b/Engine/Plugins/Compositing/OpenColorIO/Source/OpenColorIO/Private/OpenColorIOShared.cpp
--- a/Engine/Plugins/Compositing/OpenColorIO/Source/OpenColorIO/Private/OpenColorIOShared.cpp
+++ b/Engine/Plugins/Compositing/OpenColorIO/Source/OpenColorIO/Private/OpenColorIOShared.cpp
@@ -211,7 +211,10 @@ void FOpenColorIOTransformResource::SerializeShaderMap(FArchive& Ar)
 
 void FOpenColorIOTransformResource::SetupResource(ERHIFeatureLevel::Type InFeatureLevel, const FString& InShaderCodeHash, const FString& InShadercode, const FString& InRawConfigHash, const FString& InFriendlyName, const FName& InAssetPath, FShaderParametersMetadata* InShaderParameterMetadata)
 {
-	// When this happens we assume that shader was cooked and we don't need to do anything. 
+	// Take ownership of the parameter metadata
+	ShaderParametersMetadata = TSharedPtr<FShaderParametersMetadata>(InShaderParameterMetadata);
+
+	// When this happens we assume that shader was cooked and we don't need to do anything else.
 	if (InShaderCodeHash.IsEmpty() && InRawConfigHash.IsEmpty())
 	{
 		return;
@@ -225,7 +228,6 @@ void FOpenColorIOTransformResource::SetupResource(ERHIFeatureLevel::Type InFeatu
 	ShaderCode = InShadercode;
 	FriendlyName = InFriendlyName.Replace(TEXT("/"), TEXT(""));
 	AssetPath = InAssetPath;
-	ShaderParametersMetadata = TSharedPtr<FShaderParametersMetadata>(InShaderParameterMetadata);
 
 	if (!ShaderParametersMetadata.IsValid())
 	{
