Updated headers using UnrealCodeFixup to make sure dllstorage is on methods/staticvars instead of types. Part 1/n

diff --git a/Engine/Plugins/Compositing/OpenColorIO/Source/OpenColorIO/Public/OpenColorIOColorTransform.h b/Engine/Plugins/Compositing/OpenColorIO/Source/OpenColorIO/Public/OpenColorIOColorTransform.h
--- a/Engine/Plugins/Compositing/OpenColorIO/Source/OpenColorIO/Public/OpenColorIOColorTransform.h
+++ b/Engine/Plugins/Compositing/OpenColorIO/Source/OpenColorIO/Public/OpenColorIOColorTransform.h
@@ -14,6 +14,8 @@
 
 #include "OpenColorIOColorTransform.generated.h"
 
+#define UE_API OPENCOLORIO_API
+
 
 struct FImageView;
 class FOpenColorIOWrapperProcessor;
@@ -23,14 +25,14 @@ class UOpenColorIOConfiguration;
 /**
  * Object used to generate shader and LUTs from OCIO configuration file and contain required resource to make a color space transform.
  */
-UCLASS()
-class OPENCOLORIO_API UOpenColorIOColorTransform : public UObject
+UCLASS(MinimalAPI)
+class UOpenColorIOColorTransform : public UObject
 {
 	GENERATED_BODY()
 
 public:
 
-	UOpenColorIOColorTransform(const FObjectInitializer& ObjectInitializer);
+	UE_API UOpenColorIOColorTransform(const FObjectInitializer& ObjectInitializer);
 	
 	PRAGMA_DISABLE_DEPRECATION_WARNINGS
 	virtual ~UOpenColorIOColorTransform() {};
@@ -42,30 +44,30 @@ public:
 	bool Initialize(const FString& InSourceColorSpace, const FString& InDestinationColorSpace);
 	
 	/** Initialize resources for display-view transform. */
-	bool Initialize(const FString& InSourceColorSpace, const FString& InDisplay, const FString& InView, EOpenColorIOViewTransformDirection InDirection);
+	UE_API bool Initialize(const FString& InSourceColorSpace, const FString& InDisplay, const FString& InView, EOpenColorIOViewTransformDirection InDirection);
 
 	/**
 	 * Returns the desired resources required to apply this transform during rendering.
 	 */
-	bool GetRenderResources(ERHIFeatureLevel::Type InFeatureLevel, FOpenColorIOTransformResource*& OutShaderResource, TArray<TWeakObjectPtr<UTexture>>& OutTextureResources) const;
+	UE_API bool GetRenderResources(ERHIFeatureLevel::Type InFeatureLevel, FOpenColorIOTransformResource*& OutShaderResource, TArray<TWeakObjectPtr<UTexture>>& OutTextureResources) const;
 
 	UE_DEPRECATED(5.6, "This method is deprecated, please use the one with an array of textures instead.")
-	bool GetRenderResources(ERHIFeatureLevel::Type InFeatureLevel, FOpenColorIOTransformResource*& OutShaderResource, TSortedMap<int32, TWeakObjectPtr<UTexture>>& OutTextureResources) const;
+	UE_API bool GetRenderResources(ERHIFeatureLevel::Type InFeatureLevel, FOpenColorIOTransformResource*& OutShaderResource, TSortedMap<int32, TWeakObjectPtr<UTexture>>& OutTextureResources) const;
 
 	/**
 	 * Returns true if shader/texture resources have finished compiling and are ready for use (to be called on the game thread).
 	 */
-	bool AreRenderResourcesReady() const;
+	UE_API bool AreRenderResourcesReady() const;
 
 	/**
 	 * Returns true if the current transform corresponds to the specified color spaces.
 	 */
-	bool IsTransform(const FString& InSourceColorSpace, const FString& InDestinationColorSpace) const;
+	UE_API bool IsTransform(const FString& InSourceColorSpace, const FString& InDestinationColorSpace) const;
 
 	/**
 	 * Returns true if the current transform corresponds to the specified color space and display/view & direction.
 	 */
-	bool IsTransform(const FString& InSourceColorSpace, const FString& InDisplay, const FString& InView, EOpenColorIOViewTransformDirection InDirection) const;
+	UE_API bool IsTransform(const FString& InSourceColorSpace, const FString& InDisplay, const FString& InView, EOpenColorIOViewTransformDirection InDirection) const;
 
 	/**
 	 * Get the transform processor.
@@ -73,16 +75,16 @@ public:
 	 * @param OutProcessor Processor wrapper for the current transform.
 	 * @return True if the processor was created and is valid.
 	 */
-	bool GetTransformProcessor(FOpenColorIOWrapperProcessor& OutProcessor, const TMap<FString, FString>& InContextOverride = {}) const;
+	UE_API bool GetTransformProcessor(FOpenColorIOWrapperProcessor& OutProcessor, const TMap<FString, FString>& InContextOverride = {}) const;
 
 	/** Apply the color transform in-place to the specified color. */
-	bool TransformColor(FLinearColor& InOutColor) const;
+	UE_API bool TransformColor(FLinearColor& InOutColor) const;
 		
 	/** Apply the color transform in-place to the specified image. */
-	bool TransformImage(const FImageView& InOutImage) const;
+	UE_API bool TransformImage(const FImageView& InOutImage) const;
 	
 	/** Apply the color transform from the source image to the destination image. (The destination FImageView is const but what it points at is not.) */
-	bool TransformImage(const FImageView& SrcImage, const FImageView& DestImage) const;
+	UE_API bool TransformImage(const FImageView& SrcImage, const FImageView& DestImage) const;
 
 	/**
 	 * Get the display view direction type, when applicable.
@@ -90,13 +92,13 @@ public:
 	 * @param OutDirection The returned direction.
 	 * @return True if the transformation is of display-view type.
 	 */
-	bool GetDisplayViewDirection(EOpenColorIOViewTransformDirection& OutDirection) const;
+	UE_API bool GetDisplayViewDirection(EOpenColorIOViewTransformDirection& OutDirection) const;
 
 	// For all ColorTransforms, UOpenColorIOColorTransform::CacheResourceShadersForRendering
-	static void AllColorTransformsCacheResourceShadersForRendering();
+	static UE_API void AllColorTransformsCacheResourceShadersForRendering();
 
 	// Get the owner config's context key-values.
-	TMap<FString, FString> GetContextKeyValues() const;
+	UE_API TMap<FString, FString> GetContextKeyValues() const;
 
 private:
 
@@ -104,71 +106,71 @@ private:
 	 * Cache resource shaders for rendering.
 	 * If a matching shader map is not found in memory or the DDC, a new one will be compiled.
 	 */
-	void CacheResourceShadersForRendering(bool bRegenerateId = false);
+	UE_API void CacheResourceShadersForRendering(bool bRegenerateId = false);
 
-	void CacheShadersForResources(EShaderPlatform InShaderPlatform, FOpenColorIOTransformResource* InResourcesToCache, bool bApplyCompletedShaderMapForRendering, bool bIsCooking, const ITargetPlatform* TargetPlatform = nullptr);
+	UE_API void CacheShadersForResources(EShaderPlatform InShaderPlatform, FOpenColorIOTransformResource* InResourcesToCache, bool bApplyCompletedShaderMapForRendering, bool bIsCooking, const ITargetPlatform* TargetPlatform = nullptr);
 
 	/**
 	 * Helper function to serialize shader maps for the given color transform resources.
 	 */
-	static void SerializeOpenColorIOShaderMaps(const TMap<const ITargetPlatform*, TArray<FOpenColorIOTransformResource*>>* PlatformColorTransformResourcesToSavePtr, FArchive& Ar, TArray<FOpenColorIOTransformResource>&  OutLoadedResources);
+	static UE_API void SerializeOpenColorIOShaderMaps(const TMap<const ITargetPlatform*, TArray<FOpenColorIOTransformResource*>>* PlatformColorTransformResourcesToSavePtr, FArchive& Ar, TArray<FOpenColorIOTransformResource>&  OutLoadedResources);
 	
 	/**
 	 * Helper function to register serialized shader maps for the given color transform resources
 	 */
-	static void ProcessSerializedShaderMaps(UOpenColorIOColorTransform* Owner, TArray<FOpenColorIOTransformResource>& LoadedResources, FOpenColorIOTransformResource* (&OutColorTransformResourcesLoaded)[ERHIFeatureLevel::Num]);
+	static UE_API void ProcessSerializedShaderMaps(UOpenColorIOColorTransform* Owner, TArray<FOpenColorIOTransformResource>& LoadedResources, FOpenColorIOTransformResource* (&OutColorTransformResourcesLoaded)[ERHIFeatureLevel::Num]);
 	
 	/**
 	 * Returns a Guid for the LUT based on its unique identifier, name and the OCIO DDC key.
 	 */
-	static void GetOpenColorIOLUTKeyGuid(const FString& InProcessorIdentifier, const FName& InName, FGuid& OutLutGuid );
+	static UE_API void GetOpenColorIOLUTKeyGuid(const FString& InProcessorIdentifier, const FName& InName, FGuid& OutLutGuid );
 
 	/**
 	 * Helper function returning the color space transform name based on source and destination color spaces.
 	 */
-	FString GetTransformFriendlyName() const;
+	UE_API FString GetTransformFriendlyName() const;
 
 #if WITH_EDITOR
 	/**
 	 * Fetch shader code and hash from the OCIO library
 	 * @return: true if shader could be generated from the library
 	 */
-	bool UpdateShaderInfo(FString& OutShaderCodeHash, FString& OutShaderCode, FString& OutRawConfigHash);
+	UE_API bool UpdateShaderInfo(FString& OutShaderCodeHash, FString& OutShaderCode, FString& OutRawConfigHash);
 
 	/**
 	 * Helper function taking raw 3D LUT data coming from the library and initializing a UVolumeTexture with it.
 	 */
-	TObjectPtr<UTexture> CreateTexture3DLUT(const FString& InProcessorIdentifier, const FName& InName, uint32 InLutLength, TextureFilter InFilter, const float* InSourceData);
+	UE_API TObjectPtr<UTexture> CreateTexture3DLUT(const FString& InProcessorIdentifier, const FName& InName, uint32 InLutLength, TextureFilter InFilter, const float* InSourceData);
 
 	/**
 	 * Helper function taking raw 1D LUT data coming from the library and initializing a UTexture with it.
 	 */
-	TObjectPtr<UTexture> CreateTexture1DLUT(const FString& InProcessorIdentifier, const FName& InName, uint32 InTextureWidth, uint32 InTextureHeight, TextureFilter InFilter, bool bRedChannelOnly, const float* InSourceData );
+	UE_API TObjectPtr<UTexture> CreateTexture1DLUT(const FString& InProcessorIdentifier, const FName& InName, uint32 InTextureWidth, uint32 InTextureHeight, TextureFilter InFilter, bool bRedChannelOnly, const float* InSourceData );
 
 	/**
 	 * Create the transform processor(s) and generate its resources. Editor-only, in game mode GPU resources are already present.
 	 */
-	void ProcessTransformForGPU();
+	UE_API void ProcessTransformForGPU();
 #endif
 
 	/* Release all shader maps. */
-	void FlushResourceShaderMaps();
+	UE_API void FlushResourceShaderMaps();
 
 	/* Create the shader parameter metadata, using our built-in struct and any LUT-associated texture and sampler resources. */
-	FShaderParametersMetadata* BuildShaderParamMetadata();
+	UE_API FShaderParametersMetadata* BuildShaderParamMetadata();
 
 public:
 
 	//~ Begin UObject interface
-	void Serialize(FArchive& Ar) override;
-	virtual void PostLoad() override;
+	UE_API void Serialize(FArchive& Ar) override;
+	UE_API virtual void PostLoad() override;
 #if WITH_EDITORONLY_DATA
-	static void DeclareConstructClasses(TArray<FTopLevelAssetPath>& OutConstructClasses, const UClass* SpecificSubclass);
+	static UE_API void DeclareConstructClasses(TArray<FTopLevelAssetPath>& OutConstructClasses, const UClass* SpecificSubclass);
 #endif
 
-	virtual void BeginDestroy() override;
-	virtual bool IsReadyForFinishDestroy() override;
-	virtual void FinishDestroy() override;
+	UE_API virtual void BeginDestroy() override;
+	UE_API virtual bool IsReadyForFinishDestroy() override;
+	UE_API virtual void FinishDestroy() override;
 
 #if WITH_EDITOR
 	/**
@@ -177,12 +179,12 @@ public:
 	 * This does not apply completed results to the renderer scenes.
 	 * Caller is responsible for deleting OutCachedResource.
 	 */
-	void CacheResourceShadersForCooking(EShaderPlatform InShaderPlatform, const ITargetPlatform* TargetPlatform, const FString& InShaderHash, const FString& InShaderCode, const FString& InRawConfigHash, TArray<FOpenColorIOTransformResource*>& OutCachedResources);
+	UE_API void CacheResourceShadersForCooking(EShaderPlatform InShaderPlatform, const ITargetPlatform* TargetPlatform, const FString& InShaderHash, const FString& InShaderCode, const FString& InRawConfigHash, TArray<FOpenColorIOTransformResource*>& OutCachedResources);
 
-	virtual void BeginCacheForCookedPlatformData(const ITargetPlatform* TargetPlatform) override;
-	virtual bool IsCachedCookedPlatformDataLoaded(const ITargetPlatform* TargetPlatform) override;
-	virtual void ClearCachedCookedPlatformData(const ITargetPlatform *TargetPlatform) override;
-	virtual void ClearAllCachedCookedPlatformData() override;
+	UE_API virtual void BeginCacheForCookedPlatformData(const ITargetPlatform* TargetPlatform) override;
+	UE_API virtual bool IsCachedCookedPlatformDataLoaded(const ITargetPlatform* TargetPlatform) override;
+	UE_API virtual void ClearCachedCookedPlatformData(const ITargetPlatform *TargetPlatform) override;
+	UE_API virtual void ClearAllCachedCookedPlatformData() override;
 
 #endif //WITH_EDITOR
 	//~ End UObject interface
@@ -192,7 +194,7 @@ public:
 	 * This should only be called directly if the ColorTransform will not be deleted through the GC system afterward.
 	 * FlushRenderingCommands() must have been called before this.
 	 */
-	void ReleaseResources();
+	UE_API void ReleaseResources();
 
 public:
 
@@ -250,3 +252,5 @@ private:
 	FDelegateHandle FeatureLevelChangedDelegateHandle;
 #endif //WITH_EDITORONLY_DATA
 };
+
+#undef UE_API
diff --git a/Engine/Plugins/Compositing/OpenColorIO/Source/OpenColorIO/Public/OpenColorIOConfiguration.h b/Engine/Plugins/Compositing/OpenColorIO/Source/OpenColorIO/Public/OpenColorIOConfiguration.h
--- a/Engine/Plugins/Compositing/OpenColorIO/Source/OpenColorIO/Public/OpenColorIOConfiguration.h
+++ b/Engine/Plugins/Compositing/OpenColorIO/Source/OpenColorIO/Public/OpenColorIOConfiguration.h
@@ -9,6 +9,8 @@
 #include "OpenColorIOColorSpace.h"
 #include "OpenColorIOConfiguration.generated.h"
 
+#define UE_API OPENCOLORIO_API
+
 
 class FOpenColorIOWrapperConfig;
 class FOpenColorIOTransformResource;
@@ -23,116 +25,116 @@ namespace ERHIFeatureLevel { enum Type : int; }
 /**
  * Asset to manage allowed OpenColorIO color spaces. This will create required transform objects.
  */
-UCLASS(BlueprintType, meta = (DisplayName = "OpenColorIO Configuration"))
-class OPENCOLORIO_API UOpenColorIOConfiguration : public UObject
+UCLASS(MinimalAPI, BlueprintType, meta = (DisplayName = "OpenColorIO Configuration"))
+class UOpenColorIOConfiguration : public UObject
 {
 	GENERATED_BODY()
 
 public:
 
-	UOpenColorIOConfiguration(const FObjectInitializer& ObjectInitializer);
+	UE_API UOpenColorIOConfiguration(const FObjectInitializer& ObjectInitializer);
 
 public:
 
 	/** Check if the transform resources (shader & optionally lookup textures) are ready for use. */
-	bool IsTransformReady(const FOpenColorIOColorConversionSettings& InSettings);
+	UE_API bool IsTransformReady(const FOpenColorIOColorConversionSettings& InSettings);
 
 	/** Get the shader and optionally lookup texture resources to be used by the color transform render pass specified by the settings parameter. */
-	bool GetRenderResources(ERHIFeatureLevel::Type InFeatureLevel, const FOpenColorIOColorConversionSettings& InSettings, FOpenColorIOTransformResource*& OutShaderResource, TArray<TWeakObjectPtr<UTexture>>& OutTextureResources);
+	UE_API bool GetRenderResources(ERHIFeatureLevel::Type InFeatureLevel, const FOpenColorIOColorConversionSettings& InSettings, FOpenColorIOTransformResource*& OutShaderResource, TArray<TWeakObjectPtr<UTexture>>& OutTextureResources);
 
 	UE_DEPRECATED(5.6, "This method is deprecated, please use the one with an array of textures instead.")
-	bool GetRenderResources(ERHIFeatureLevel::Type InFeatureLevel, const FOpenColorIOColorConversionSettings& InSettings, FOpenColorIOTransformResource*& OutShaderResource, TSortedMap<int32, TWeakObjectPtr<UTexture>>& OutTextureResources);
+	UE_API bool GetRenderResources(ERHIFeatureLevel::Type InFeatureLevel, const FOpenColorIOColorConversionSettings& InSettings, FOpenColorIOTransformResource*& OutShaderResource, TSortedMap<int32, TWeakObjectPtr<UTexture>>& OutTextureResources);
 	
 	/** Check if this configuration contains a transform from the specified source & destination color spaces. */
-	bool HasTransform(const FString& InSourceColorSpace, const FString& InDestinationColorSpace);
+	UE_API bool HasTransform(const FString& InSourceColorSpace, const FString& InDestinationColorSpace);
 	
 	/** Check if this configuration contains a transform from the specified source color space & view destination. */
-	bool HasTransform(const FString& InSourceColorSpace, const FString& InDisplay, const FString& InView, EOpenColorIOViewTransformDirection InDirection);
+	UE_API bool HasTransform(const FString& InSourceColorSpace, const FString& InDisplay, const FString& InView, EOpenColorIOViewTransformDirection InDirection);
 
 	/** Check if this configuration contains a color space. */
-	bool HasDesiredColorSpace(const FOpenColorIOColorSpace& ColorSpace) const;
+	UE_API bool HasDesiredColorSpace(const FOpenColorIOColorSpace& ColorSpace) const;
 
 	/** Check if this configuration contains a display-view. */
-	bool HasDesiredDisplayView(const FOpenColorIODisplayView& DisplayView) const;
+	UE_API bool HasDesiredDisplayView(const FOpenColorIODisplayView& DisplayView) const;
 
 	/** Check if this configuration is valid. (TODO: Rename to avoid confusion with config validation calls provided by the library.) */
-	bool Validate() const;
+	UE_API bool Validate() const;
 
 	/** Apply the color transform in-place to the specified image. */
-	bool TransformColor(const FOpenColorIOColorConversionSettings& InSettings, FLinearColor& InOutColor) const;
+	UE_API bool TransformColor(const FOpenColorIOColorConversionSettings& InSettings, FLinearColor& InOutColor) const;
 
 	/** Apply the color transform in-place to the specified image. */
-	bool TransformImage(const FOpenColorIOColorConversionSettings& InSettings, const FImageView& InOutImage) const;
+	UE_API bool TransformImage(const FOpenColorIOColorConversionSettings& InSettings, const FImageView& InOutImage) const;
 
 	/** Apply the color transform from the source image to the destination image. (The destination FImageView is const but what it points at is not.) */
-	bool TransformImage(const FOpenColorIOColorConversionSettings& InSettings, const FImageView& SrcImage, const FImageView& DestImage) const;
+	UE_API bool TransformImage(const FOpenColorIOColorConversionSettings& InSettings, const FImageView& SrcImage, const FImageView& DestImage) const;
 
 	/** This forces to reload colorspaces and corresponding shaders if those are not loaded already. */
 	UFUNCTION(BlueprintCallable, Category = "OpenColorIO")
-	void ReloadExistingColorspaces(bool bForce = false);
+	UE_API void ReloadExistingColorspaces(bool bForce = false);
 
 	/*
 	* This method is called by directory watcher when any file or folder is changed in the 
 	* directory where raw ocio config is located.
 	*/
-	void ConfigPathChangedEvent(const TArray<FFileChangeData>& InFileChanges, const FString InFileMountPath);
+	UE_API void ConfigPathChangedEvent(const TArray<FFileChangeData>& InFileChanges, const FString InFileMountPath);
 
 	/**
 	 * Get the private wrapper implementation of the OCIO config.
 	 */
-	FOpenColorIOWrapperConfig* GetConfigWrapper() const;
+	UE_API FOpenColorIOWrapperConfig* GetConfigWrapper() const;
 	
 	/**
 	 * Get or create the private wrapper implementation of the OCIO config.
 	 * Useful for non-editor modes where the config isn't automatically loaded.
 	 */
-	FOpenColorIOWrapperConfig* GetOrCreateConfigWrapper();
+	UE_API FOpenColorIOWrapperConfig* GetOrCreateConfigWrapper();
 	
 	/** Find the color transform object that corresponds to the specified settings, nullptr if not found. */
-	TObjectPtr<const UOpenColorIOColorTransform> FindTransform(const FOpenColorIOColorConversionSettings& InSettings) const;
+	UE_API TObjectPtr<const UOpenColorIOColorTransform> FindTransform(const FOpenColorIOColorConversionSettings& InSettings) const;
 
 protected:
 
-	void CreateColorTransform(const FString& InSourceColorSpace, const FString& InDestinationColorSpace);
-	void CreateColorTransform(const FString& InSourceColorSpace, const FString& InDisplay, const FString& InView, EOpenColorIOViewTransformDirection InDirection);
+	UE_API void CreateColorTransform(const FString& InSourceColorSpace, const FString& InDestinationColorSpace);
+	UE_API void CreateColorTransform(const FString& InSourceColorSpace, const FString& InDisplay, const FString& InView, EOpenColorIOViewTransformDirection InDirection);
 
-	void CleanupTransforms();
+	UE_API void CleanupTransforms();
 
 	/** Same as above except user can specify the path manually. */
-	void StartDirectoryWatch(const FString& FilePath);
+	UE_API void StartDirectoryWatch(const FString& FilePath);
 
 	/** Stop watching the current directory. */
-	void StopDirectoryWatch();
+	UE_API void StopDirectoryWatch();
 
 public:
 
 	//~ Begin UObject interface
-	virtual void PostInitProperties() override;
-	virtual void PostLoad() override;
-	virtual void GetAssetRegistryTags(FAssetRegistryTagsContext RegistryTagsContext) const override;
+	UE_API virtual void PostInitProperties() override;
+	UE_API virtual void PostLoad() override;
+	UE_API virtual void GetAssetRegistryTags(FAssetRegistryTagsContext RegistryTagsContext) const override;
 #if WITH_EDITORONLY_DATA
-	static void DeclareConstructClasses(TArray<FTopLevelAssetPath>& OutConstructClasses, const UClass* SpecificSubclass);
+	static UE_API void DeclareConstructClasses(TArray<FTopLevelAssetPath>& OutConstructClasses, const UClass* SpecificSubclass);
 #endif
 
-	virtual void PreSave(FObjectPreSaveContext SaveContext) override;
-	virtual void BeginDestroy() override;
+	UE_API virtual void PreSave(FObjectPreSaveContext SaveContext) override;
+	UE_API virtual void BeginDestroy() override;
 
 #if WITH_EDITOR
-	virtual void PostEditChangeProperty(FPropertyChangedEvent& PropertyChangedEvent) override;
+	UE_API virtual void PostEditChangeProperty(FPropertyChangedEvent& PropertyChangedEvent) override;
 #endif //WITH_EDITOR
 	//~ End UObject interface
 
 private:
 
 	/** Load the config file to initialize the configuration wrapper object. Automatically called internally. */
-	void LoadConfiguration();
+	UE_API void LoadConfiguration();
 
 	/** Obtain the hash of the config and additional dependent state such as library version, working color space and context. */
-	bool GetHash(FString& OutHash) const;
+	UE_API bool GetHash(FString& OutHash) const;
 
 #if WITH_EDITOR
 	/** This method resets the status of Notification dialog and reacts depending on user's choice. */
-	void OnToastCallback(bool bInReloadColorspaces);
+	UE_API void OnToastCallback(bool bInReloadColorspaces);
 #endif
 
 public:
@@ -181,3 +183,4 @@ private:
 	FString ConfigHash;
 };
 
+#undef UE_API
diff --git a/Engine/Plugins/Compositing/OpenColorIO/Source/OpenColorIO/Public/OpenColorIODisplayExtension.h b/Engine/Plugins/Compositing/OpenColorIO/Source/OpenColorIO/Public/OpenColorIODisplayExtension.h
--- a/Engine/Plugins/Compositing/OpenColorIO/Source/OpenColorIO/Public/OpenColorIODisplayExtension.h
+++ b/Engine/Plugins/Compositing/OpenColorIO/Source/OpenColorIO/Public/OpenColorIODisplayExtension.h
@@ -9,6 +9,8 @@
 #include "OpenColorIORendering.h"
 #include "UObject/GCObject.h"
 
+#define UE_API OPENCOLORIO_API
+
 
 class FViewportClient;
 class FSceneViewFamily;
@@ -18,24 +20,24 @@ class FSceneViewFamily;
 /** 
  * View extension applying an OCIO Display Look to the viewport we're attached to
  */
-class OPENCOLORIO_API FOpenColorIODisplayExtension : public FSceneViewExtensionBase, public FGCObject
+class FOpenColorIODisplayExtension : public FSceneViewExtensionBase, public FGCObject
 {
 public:
 
-	FOpenColorIODisplayExtension(const FAutoRegister& AutoRegister, FViewportClient* AssociatedViewportClient = nullptr);
+	UE_API FOpenColorIODisplayExtension(const FAutoRegister& AutoRegister, FViewportClient* AssociatedViewportClient = nullptr);
 	
 	//~ Begin ISceneViewExtension interface	
 	virtual int32 GetPriority() const override { return OPENCOLORIO_SCENE_VIEW_EXTENSION_PRIORITY; }
-	virtual void SetupView(FSceneViewFamily& InViewFamily, FSceneView& InView) override;
+	UE_API virtual void SetupView(FSceneViewFamily& InViewFamily, FSceneView& InView) override;
 
-	virtual void SubscribeToPostProcessingPass(EPostProcessingPass PassId, const FSceneView& View, FAfterPassCallbackDelegateArray& InOutPassCallbacks, bool bIsPassEnabled) override;
-	virtual bool IsActiveThisFrame_Internal(const FSceneViewExtensionContext& Context) const override;
+	UE_API virtual void SubscribeToPostProcessingPass(EPostProcessingPass PassId, const FSceneView& View, FAfterPassCallbackDelegateArray& InOutPassCallbacks, bool bIsPassEnabled) override;
+	UE_API virtual bool IsActiveThisFrame_Internal(const FSceneViewExtensionContext& Context) const override;
 	//~End ISceneVIewExtension interface
 
-	FScreenPassTexture PostProcessPassAfterTonemap_RenderThread(FRDGBuilder& GraphBuilder, const FSceneView& View, const FPostProcessMaterialInputs& Inputs);
+	UE_API FScreenPassTexture PostProcessPassAfterTonemap_RenderThread(FRDGBuilder& GraphBuilder, const FSceneView& View, const FPostProcessMaterialInputs& Inputs);
 
 	//~Begin FGCObject interface
-	virtual void AddReferencedObjects(FReferenceCollector& Collector) override;
+	UE_API virtual void AddReferencedObjects(FReferenceCollector& Collector) override;
 	virtual FString GetReferencerName() const override
 	{
 		return TEXT("FOpenColorIODisplayExtension");
@@ -67,3 +69,4 @@ private:
 };
 
 
+#undef UE_API
diff --git a/Engine/Plugins/Compositing/OpenColorIO/Source/OpenColorIO/Public/OpenColorIODisplayExtensionWrapper.h b/Engine/Plugins/Compositing/OpenColorIO/Source/OpenColorIO/Public/OpenColorIODisplayExtensionWrapper.h
--- a/Engine/Plugins/Compositing/OpenColorIO/Source/OpenColorIO/Public/OpenColorIODisplayExtensionWrapper.h
+++ b/Engine/Plugins/Compositing/OpenColorIO/Source/OpenColorIO/Public/OpenColorIODisplayExtensionWrapper.h
@@ -7,12 +7,14 @@
 
 #include "OpenColorIODisplayExtensionWrapper.generated.h"
 
+#define UE_API OPENCOLORIO_API
+
 /** 
  * This Blueprintable object can hold an OCIO Scene View Extension. 
  * You can change its OCIO config, and specify the context in which you want it to be active on.
  */
-UCLASS(Blueprintable, meta = (DisplayName = "OpenColorIO Display Extension Wrapper"))
-class OPENCOLORIO_API UOpenColorIODisplayExtensionWrapper : public UObject
+UCLASS(MinimalAPI, Blueprintable, meta = (DisplayName = "OpenColorIO Display Extension Wrapper"))
+class UOpenColorIODisplayExtensionWrapper : public UObject
 {
 	GENERATED_BODY()
 
@@ -20,38 +22,38 @@ public:
 
 	// Gets the display extension OCIO configuration.
 	UFUNCTION(BlueprintCallable, meta = (DisplayName = "Get OpenColorIO Configuration"), Category = "OpenColorIO")
-	FOpenColorIODisplayConfiguration GetOpenColorIOConfiguration() const;
+	UE_API FOpenColorIODisplayConfiguration GetOpenColorIOConfiguration() const;
 
 	// Sets the display extension OCIO configuration.
 	UFUNCTION(BlueprintCallable, meta = (DisplayName = "Set OpenColorIO Configuration"), Category = "OpenColorIO")
-	void SetOpenColorIOConfiguration(FOpenColorIODisplayConfiguration InDisplayConfiguration);
+	UE_API void SetOpenColorIOConfiguration(FOpenColorIODisplayConfiguration InDisplayConfiguration);
 
 	// Sets a single activation function. Will remove any others.
 	UFUNCTION(BlueprintCallable, Category = "OpenColorIO")
-	void SetSceneExtensionIsActiveFunction(const FSceneViewExtensionIsActiveFunctor& IsActiveFunction);
+	UE_API void SetSceneExtensionIsActiveFunction(const FSceneViewExtensionIsActiveFunctor& IsActiveFunction);
 
 	// Sets an array of activation functions. Will remove any others.
 	UFUNCTION(BlueprintCallable, Category = "OpenColorIO")
-	void SetSceneExtensionIsActiveFunctions(const TArray<FSceneViewExtensionIsActiveFunctor>& IsActiveFunctions);
+	UE_API void SetSceneExtensionIsActiveFunctions(const TArray<FSceneViewExtensionIsActiveFunctor>& IsActiveFunctions);
 
 	// Removes the extension.
 	UFUNCTION(BlueprintCallable, Category = "OpenColorIO")
-	void RemoveSceneExtension();
+	UE_API void RemoveSceneExtension();
 
 	// Creates an instance of this object, configured with the given arguments (OCIO and activation function).
 	UFUNCTION(BlueprintCallable, meta = (DisplayName = "Create OpenColorIO Display Extension"), Category = "OpenColorIO")
-	static UOpenColorIODisplayExtensionWrapper* CreateOpenColorIODisplayExtension(
+	static UE_API UOpenColorIODisplayExtensionWrapper* CreateOpenColorIODisplayExtension(
 		FOpenColorIODisplayConfiguration InDisplayConfiguration, 
 		const FSceneViewExtensionIsActiveFunctor& IsActiveFunction);
 
 	// Creates an instance of this object, configured for use in game with the given OCIO configuration.
 	UFUNCTION(BlueprintCallable, meta = (DisplayName = "Create In-Game OpenColorIO Display Extension"), Category = "OpenColorIO")
-	static UOpenColorIODisplayExtensionWrapper* CreateInGameOpenColorIODisplayExtension(
+	static UE_API UOpenColorIODisplayExtensionWrapper* CreateInGameOpenColorIODisplayExtension(
 		FOpenColorIODisplayConfiguration InDisplayConfiguration);
 private:
 
 	// Creates OCIO Extension if it doesn't exist. 
-	void CreateDisplayExtensionIfNotExists();
+	UE_API void CreateDisplayExtensionIfNotExists();
 
 private:
 
@@ -60,3 +62,4 @@ private:
 };
 
 
+#undef UE_API
diff --git a/Engine/Plugins/Compositing/OpenColorIO/Source/OpenColorIO/Public/OpenColorIODisplayManager.h b/Engine/Plugins/Compositing/OpenColorIO/Source/OpenColorIO/Public/OpenColorIODisplayManager.h
--- a/Engine/Plugins/Compositing/OpenColorIO/Source/OpenColorIO/Public/OpenColorIODisplayManager.h
+++ b/Engine/Plugins/Compositing/OpenColorIO/Source/OpenColorIO/Public/OpenColorIODisplayManager.h
@@ -7,24 +7,26 @@
 #include "CoreMinimal.h"
 #include "OpenColorIOColorSpace.h"
 
+#define UE_API OPENCOLORIO_API
+
 class FOpenColorIODisplayExtension;
 class FViewportClient;
 
-class OPENCOLORIO_API FOpenColorIODisplayManager : public TSharedFromThis<FOpenColorIODisplayManager>
+class FOpenColorIODisplayManager : public TSharedFromThis<FOpenColorIODisplayManager>
 {
 public:
 
 	/** Will return the configuration associated to the desired viewport or create one if it's not tracked */
-	FOpenColorIODisplayConfiguration& FindOrAddDisplayConfiguration(FViewportClient* InViewportClient);
+	UE_API FOpenColorIODisplayConfiguration& FindOrAddDisplayConfiguration(FViewportClient* InViewportClient);
 
 	/** Returns the configuration for a given viewport if it was found, nullptr otherwise */
-	const FOpenColorIODisplayConfiguration* GetDisplayConfiguration(const FViewportClient* InViewportClient) const;
+	UE_API const FOpenColorIODisplayConfiguration* GetDisplayConfiguration(const FViewportClient* InViewportClient) const;
 
 	/** Remove display configuration associated with this viewport */
-	bool RemoveDisplayConfiguration(const FViewportClient* InViewportClient);
+	UE_API bool RemoveDisplayConfiguration(const FViewportClient* InViewportClient);
 
 	/** Whether or not InViewport has a display configuration linked to it */
-	bool IsTrackingViewport(const FViewportClient* InViewportClient) const;
+	UE_API bool IsTrackingViewport(const FViewportClient* InViewportClient) const;
 
 protected:
 
@@ -32,3 +34,4 @@ protected:
 	TArray<TSharedPtr<FOpenColorIODisplayExtension, ESPMode::ThreadSafe>> DisplayExtensions;
 };
 
+#undef UE_API
diff --git a/Engine/Plugins/Compositing/OpenColorIO/Source/OpenColorIO/Public/OpenColorIORendering.h b/Engine/Plugins/Compositing/OpenColorIO/Source/OpenColorIO/Public/OpenColorIORendering.h
--- a/Engine/Plugins/Compositing/OpenColorIO/Source/OpenColorIO/Public/OpenColorIORendering.h
+++ b/Engine/Plugins/Compositing/OpenColorIO/Source/OpenColorIO/Public/OpenColorIORendering.h
@@ -6,6 +6,8 @@
 #include "UObject/WeakObjectPtr.h"
 #include "RHIFwd.h"
 
+#define UE_API OPENCOLORIO_API
+
 struct FOpenColorIOColorConversionSettings;
 class FOpenColorIOTransformResource;
 class FRDGBuilder;
@@ -20,7 +22,7 @@ class UTextureRenderTarget2D;
 class UWorld;
 
 /** Resources needed by the FOpenColorIORendering pass function. */
-struct OPENCOLORIO_API FOpenColorIORenderPassResources
+struct FOpenColorIORenderPassResources
 {
 	/** Color transform pass (generated) shader. */
 	FOpenColorIOTransformResource* ShaderResource = nullptr;
@@ -44,7 +46,7 @@ enum class EOpenColorIOTransformAlpha : uint32
 };
 
 /** Entry point to trigger OpenColorIO conversion rendering */
-class OPENCOLORIO_API FOpenColorIORendering
+class FOpenColorIORendering
 {
 public:
 	FOpenColorIORendering() = delete;
@@ -58,7 +60,7 @@ public:
 	 * @param OutRenderTarget RenderTarget where to draw the input texture in the destination color space
 	 * @return True if a rendering command to apply the transform was queued.
 	 */
-	static bool ApplyColorTransform(UWorld* InWorld, const FOpenColorIOColorConversionSettings& InSettings, UTexture* InTexture, UTextureRenderTarget2D* OutRenderTarget);
+	static UE_API bool ApplyColorTransform(UWorld* InWorld, const FOpenColorIOColorConversionSettings& InSettings, UTexture* InTexture, UTextureRenderTarget2D* OutRenderTarget);
 
 
 	/**
@@ -68,7 +70,7 @@ public:
 	 * @param InFeatureLevel Rendering feature level
 	 * @return FOpenColorIORenderPassResources Render pass resources to be enqueue-copied onto the render thread.
 	 */
-	static FOpenColorIORenderPassResources GetRenderPassResources(const FOpenColorIOColorConversionSettings& InSettings, ERHIFeatureLevel::Type InFeatureLevel);
+	static UE_API FOpenColorIORenderPassResources GetRenderPassResources(const FOpenColorIOColorConversionSettings& InSettings, ERHIFeatureLevel::Type InFeatureLevel);
 
 	/**
 	 * Applies the color transform RDG pass with the provided resources.
@@ -82,7 +84,7 @@ public:
 	 * @param InGamma Display gamma
 	 * @param TransformAlpha Whether to unpremult/invert before applying the color transform
 	 */
-	static void AddPass_RenderThread(
+	static UE_API void AddPass_RenderThread(
 		FRDGBuilder& GraphBuilder,
 		FScreenPassViewInfo ViewInfo,
 		ERHIFeatureLevel::Type FeatureLevel,
@@ -102,7 +104,7 @@ public:
 	 * @param InPassInfo OpenColorIO shader and texture resources
 	 * @param TransformAlpha Whether to unpremult/invert before applying the color transform
 	 */
-	static void AddPass_RenderThread(
+	static UE_API void AddPass_RenderThread(
 		FRDGBuilder& GraphBuilder,
 		const FSceneView& View,
 		const FScreenPassTexture& Input,
@@ -113,8 +115,10 @@ public:
 	/** 
 	* This function sets up provided view and related family to properly support OCIO. 
 	*/
-	static void PrepareView(FSceneViewFamily& InViewFamily, FSceneView& InView);
+	static UE_API void PrepareView(FSceneViewFamily& InViewFamily, FSceneView& InView);
 	
 	UE_DEPRECATED(5.6, "This constant is deprecated, please use UTextureRenderTarget::GetDefaultDisplayGamma() instead.")
-	static float DefaultDisplayGamma;
-};
\ No newline at end of file
+	static UE_API float DefaultDisplayGamma;
+};
+
+#undef UE_API
diff --git a/Engine/Plugins/Compositing/OpenColorIO/Source/OpenColorIO/Public/OpenColorIOSettings.h b/Engine/Plugins/Compositing/OpenColorIO/Source/OpenColorIO/Public/OpenColorIOSettings.h
--- a/Engine/Plugins/Compositing/OpenColorIO/Source/OpenColorIO/Public/OpenColorIOSettings.h
+++ b/Engine/Plugins/Compositing/OpenColorIO/Source/OpenColorIO/Public/OpenColorIOSettings.h
@@ -9,22 +9,24 @@
 
 #include "OpenColorIOSettings.generated.h"
 
+#define UE_API OPENCOLORIO_API
+
 /**
  * Rendering settings.
  */
-UCLASS(config = Engine, defaultconfig)
-class OPENCOLORIO_API UOpenColorIOSettings : public UDeveloperSettings
+UCLASS(MinimalAPI, config = Engine, defaultconfig)
+class UOpenColorIOSettings : public UDeveloperSettings
 {
 	GENERATED_BODY()
 
 public:
-	UOpenColorIOSettings();
+	UE_API UOpenColorIOSettings();
 
 	//~ Begin UDeveloperSettings interface
-	virtual FName GetCategoryName() const;
+	UE_API virtual FName GetCategoryName() const;
 #if WITH_EDITOR
-	virtual FText GetSectionText() const override;
-	virtual FName GetSectionName() const override;
+	UE_API virtual FText GetSectionText() const override;
+	UE_API virtual FName GetSectionName() const override;
 #endif
 	//~ End UDeveloperSettings interface
 
@@ -48,3 +50,5 @@ public:
 		ConfigRestartRequired = true))
 	uint8 bSupportInverseViewTransforms : 1;
 };
+
+#undef UE_API
diff --git a/Engine/Plugins/Compositing/OpenColorIO/Source/OpenColorIO/Public/OpenColorIOShader.h b/Engine/Plugins/Compositing/OpenColorIO/Source/OpenColorIO/Public/OpenColorIOShader.h
--- a/Engine/Plugins/Compositing/OpenColorIO/Source/OpenColorIO/Public/OpenColorIOShader.h
+++ b/Engine/Plugins/Compositing/OpenColorIO/Source/OpenColorIO/Public/OpenColorIOShader.h
@@ -54,10 +54,10 @@ BEGIN_SHADER_PARAMETER_STRUCT(FOpenColorIOInvalidShaderParameters, )
 	RENDER_TARGET_BINDING_SLOTS()
 END_SHADER_PARAMETER_STRUCT()
 
-class OPENCOLORIO_API FOpenColorIOInvalidPixelShader : public FGlobalShader
+class FOpenColorIOInvalidPixelShader : public FGlobalShader
 {
 public:
-	DECLARE_GLOBAL_SHADER(FOpenColorIOInvalidPixelShader);
+	DECLARE_EXPORTED_GLOBAL_SHADER(FOpenColorIOInvalidPixelShader, OPENCOLORIO_API);
 	SHADER_USE_PARAMETER_STRUCT(FOpenColorIOInvalidPixelShader, FGlobalShader);
 
 	using FParameters = FOpenColorIOInvalidShaderParameters;
diff --git a/Engine/Plugins/Compositing/OpenColorIO/Source/OpenColorIO/Public/OpenColorIOShared.h b/Engine/Plugins/Compositing/OpenColorIO/Source/OpenColorIO/Public/OpenColorIOShared.h
--- a/Engine/Plugins/Compositing/OpenColorIO/Source/OpenColorIO/Public/OpenColorIOShared.h
+++ b/Engine/Plugins/Compositing/OpenColorIO/Source/OpenColorIO/Public/OpenColorIOShared.h
@@ -314,29 +314,29 @@ private:
 /**
  * FOpenColorIOTransformResource represents a OpenColorIO color transform to the shader compilation process
  */
-class OPENCOLORIO_API FOpenColorIOTransformResource
+class FOpenColorIOTransformResource
 {
 public:
 
 	/**
 	 * Minimal initialization constructor.
 	 */
-	FOpenColorIOTransformResource();
+	OPENCOLORIO_API FOpenColorIOTransformResource();
 
 	/**
 	 * Destructor
 	 */
-	virtual ~FOpenColorIOTransformResource();
+	OPENCOLORIO_API virtual ~FOpenColorIOTransformResource();
 
 	/**
 	 * Caches the shaders for this color transform with no static parameters on the given platform.
 	 * This is used by UOpenColorIOColorTransform
 	 */
-	bool CacheShaders(EShaderPlatform InPlatform, const ITargetPlatform* TargetPlatform, bool bApplyCompletedShaderMapForRendering, bool bSynchronous);
-	bool CacheShaders(const FOpenColorIOShaderMapId& InShaderMapId, EShaderPlatform InPlatform, bool bApplyCompletedShaderMapForRendering, bool bSynchronous);
+	OPENCOLORIO_API bool CacheShaders(EShaderPlatform InPlatform, const ITargetPlatform* TargetPlatform, bool bApplyCompletedShaderMapForRendering, bool bSynchronous);
+	OPENCOLORIO_API bool CacheShaders(const FOpenColorIOShaderMapId& InShaderMapId, EShaderPlatform InPlatform, bool bApplyCompletedShaderMapForRendering, bool bSynchronous);
 
 	/** Returns the shader parameter metadata for the generated shader (built-in FOpenColorIOPixelShaderParameters, followed by dynamic textures & samplers). */
-	TSharedPtr<const FShaderParametersMetadata, ESPMode::ThreadSafe> GetShaderParamMetadata() const;
+	OPENCOLORIO_API TSharedPtr<const FShaderParametersMetadata, ESPMode::ThreadSafe> GetShaderParamMetadata() const;
 
 	/**
 	 * Should the shader for this color transform with the given platform, shader type and vertex
@@ -347,20 +347,20 @@ public:
 	 *
 	 * @return true if the shader should be compiled
 	 */
-	virtual bool ShouldCache(EShaderPlatform InPlatform, const FShaderType* InShaderType) const;
+	OPENCOLORIO_API virtual bool ShouldCache(EShaderPlatform InPlatform, const FShaderType* InShaderType) const;
 
-	void SerializeShaderMap(FArchive& Ar);
+	OPENCOLORIO_API void SerializeShaderMap(FArchive& Ar);
 
 	/** Releases this color transform's shader map.  Must only be called on ColorTransforms not exposed to the rendering thread! */
-	void ReleaseShaderMap();
+	OPENCOLORIO_API void ReleaseShaderMap();
 
 	/** Discards loaded shader maps if the application can't render */
-	void DiscardShaderMap();
+	OPENCOLORIO_API void DiscardShaderMap();
 
-	void GetDependentShaderTypes(EShaderPlatform InPlatform, TArray<FShaderType*>& OutShaderTypes) const;
-	virtual void GetShaderMapId(EShaderPlatform InPlatform, const ITargetPlatform* TargetPlatform, FOpenColorIOShaderMapId& OutId) const;
+	OPENCOLORIO_API void GetDependentShaderTypes(EShaderPlatform InPlatform, TArray<FShaderType*>& OutShaderTypes) const;
+	OPENCOLORIO_API virtual void GetShaderMapId(EShaderPlatform InPlatform, const ITargetPlatform* TargetPlatform, FOpenColorIOShaderMapId& OutId) const;
 
-	void Invalidate();
+	OPENCOLORIO_API void Invalidate();
 
 	/**
 	 * Should shaders compiled for this color transform be saved to disk?
@@ -371,12 +371,12 @@ public:
 	/**
 	 * Cancels all outstanding compilation jobs
 	 */
-	void CancelCompilation();
+	OPENCOLORIO_API void CancelCompilation();
 
 	/**
 	 * Blocks until compilation has completed. Returns immediately if a compilation is not outstanding.
 	 */
-	void FinishCompilation();
+	OPENCOLORIO_API void FinishCompilation();
 #endif // WITH_EDITOR
 
 	/**
@@ -384,7 +384,7 @@ public:
 	 *
 	 * @return returns true if compilation is complete false otherwise
 	 */
-	bool IsCompilationFinished() const;
+	OPENCOLORIO_API bool IsCompilationFinished() const;
 
 	// Accessors.
 	const TArray<FString>& GetCompileErrors() const { return CompileErrors; }
@@ -416,7 +416,7 @@ public:
 	}
 
 	/** Note: SetGameThreadShaderMap must also be called with the same value, but from the game thread. */
-	void SetRenderingThreadShaderMap(FOpenColorIOShaderMap* InShaderMap);
+	OPENCOLORIO_API void SetRenderingThreadShaderMap(FOpenColorIOShaderMap* InShaderMap);
 
 #if WITH_EDITOR
 	void AddCompileId(uint32 InIdentifier) 
@@ -439,7 +439,7 @@ public:
 	}
 
 #if WITH_EDITOR
-	void RemoveOutstandingCompileId(const int32 InOldOutstandingCompileShaderMapId);
+	OPENCOLORIO_API void RemoveOutstandingCompileId(const int32 InOldOutstandingCompileShaderMapId);
 #endif // WITH_EDITOR
 
 	/**
@@ -466,7 +466,7 @@ public:
 	) { }
 
 	// Note: Takes ownership of the shader parameter metadata.
-	void SetupResource(
+	OPENCOLORIO_API void SetupResource(
 		ERHIFeatureLevel::Type InFeatureLevel,
 		const FString& InShaderCodeHash,
 		const FString& InShadercode,
@@ -507,14 +507,14 @@ public:
 	};
 
 	
-	bool IsSame(const FOpenColorIOShaderMapId& InId) const;
+	OPENCOLORIO_API bool IsSame(const FOpenColorIOShaderMapId& InId) const;
 
 protected:
 #if WITH_EDITOR
 	/**
 	 * Fills the passed array with IDs of shader maps unfinished compilation jobs.
 	 */
-	void GetShaderMapIDsWithUnfinishedCompilation(TArray<int32>& OutShaderMapIds);
+	OPENCOLORIO_API void GetShaderMapIDsWithUnfinishedCompilation(TArray<int32>& OutShaderMapIds);
 #endif // WITH_EDITOR
 
 	void SetFeatureLevel(ERHIFeatureLevel::Type InFeatureLevel)
@@ -575,7 +575,7 @@ private:
 	/**
 	 * Compiles this color transform for InPlatform, storing the result in OutShaderMap if the compile was synchronous
 	 */
-	bool BeginCompileShaderMap(
+	OPENCOLORIO_API bool BeginCompileShaderMap(
 		const FOpenColorIOShaderMapId& InShaderMapId,
 		EShaderPlatform InPlatform, 
 		TRefCountPtr<class FOpenColorIOShaderMap>& OutShaderMap, 
@@ -583,7 +583,7 @@ private:
 		bool bSynchronous = false);
 
 	/** Populates OutEnvironment with defines needed to compile shaders for this color transform. */
-	void SetupShaderCompilationEnvironment(
+	OPENCOLORIO_API void SetupShaderCompilationEnvironment(
 		EShaderPlatform InPlatform,
 		FShaderCompilerEnvironment& OutEnvironment
 		) const;
