Removing some stale shader compilation debug logging (DEBUG_INFINITESHADERCOMPILE) that has not worked for quite a while

diff --git a/Engine/Plugins/Compositing/OpenColorIO/Source/OpenColorIO/Private/OpenColorIOShaderMap.cpp b/Engine/Plugins/Compositing/OpenColorIO/Source/OpenColorIO/Private/OpenColorIOShaderMap.cpp
--- a/Engine/Plugins/Compositing/OpenColorIO/Source/OpenColorIO/Private/OpenColorIOShaderMap.cpp
+++ b/Engine/Plugins/Compositing/OpenColorIO/Source/OpenColorIO/Private/OpenColorIOShaderMap.cpp
@@ -364,9 +364,6 @@ void FOpenColorIOShaderMap::Compile(FOpenColorIOTransformResource* InColorTransf
 			TArray<FOpenColorIOTransformResource*> NewCorrespondingTransforms;
 			NewCorrespondingTransforms.Add(InColorTransform);
 			OpenColorIOShaderMapsBeingCompiled.Add(this, NewCorrespondingTransforms);
-#if DEBUG_INFINITESHADERCOMPILE
-			UE_LOG(LogTemp, Display, TEXT("Added OpenColorIO ShaderMap 0x%08X%08X with ColorTransform 0x%08X%08X to OpenColorIOShaderMapsBeingCompiled"), (int)((int64)(this) >> 32), (int)((int64)(this)), (int)((int64)(InColorTransform) >> 32), (int)((int64)(InColorTransform)));
-#endif  
 			// Setup the compilation environment.
 			InColorTransform->SetupShaderCompilationEnvironment(InPlatform, *InCompilationEnvironment);
   
@@ -501,9 +498,6 @@ bool FOpenColorIOShaderMap::TryToAddToExistingCompilationTask(FOpenColorIOTransf
 
 		UE_LOG(LogShaders, Log, TEXT("TryToAddToExistingCompilationTask %p %d"), InColorTransform, GetCompilingId());
 
-#if DEBUG_INFINITESHADERCOMPILE
-		UE_LOG(LogTemp, Display, TEXT("Added shader map 0x%08X%08X from OpenColorIO transform 0x%08X%08X"), (int)((int64)(this) >> 32), (int)((int64)(this)), (int)((int64)(InColorTransform) >> 32), (int)((int64)(InColorTransform)));
-#endif
 		return true;
 	}
 
@@ -650,12 +644,6 @@ void FOpenColorIOShaderMap::RemovePendingColorTransform(FOpenColorIOTransformRes
 		{
 			InColorTransform->RemoveOutstandingCompileId(It.Key()->CompilingId);
 		}
-#if DEBUG_INFINITESHADERCOMPILE
-		if ( Result )
-		{
-			UE_LOG(LogTemp, Display, TEXT("Removed shader map 0x%08X%08X from color transform 0x%08X%08X"), (int)((int64)(It.Key().GetReference()) >> 32), (int)((int64)(It.Key().GetReference())), (int)((int64)(InColorTransform) >> 32), (int)((int64)(InColorTransform)));
-		}
-#endif
 	}
 }
 #endif // WITH_EDITOR
diff --git a/Engine/Plugins/Compositing/OpenColorIO/Source/OpenColorIO/Private/OpenColorIOShared.cpp b/Engine/Plugins/Compositing/OpenColorIO/Source/OpenColorIO/Private/OpenColorIOShared.cpp
--- a/Engine/Plugins/Compositing/OpenColorIO/Source/OpenColorIO/Private/OpenColorIOShared.cpp
+++ b/Engine/Plugins/Compositing/OpenColorIO/Source/OpenColorIO/Private/OpenColorIOShared.cpp
@@ -326,9 +326,6 @@ bool FOpenColorIOTransformResource::CacheShaders(const FOpenColorIOShaderMapId&
 
 	if (GameThreadShaderMap && GameThreadShaderMap->TryToAddToExistingCompilationTask(this))
 	{
-#if DEBUG_INFINITESHADERCOMPILE
-		UE_LOG(LogTemp, Display, TEXT("Found existing compiling shader for OCIO ColorTransform %s, linking to other GameThreadShaderMap 0x%08X%08X"), *GetFriendlyName(), (int)((int64)(GameThreadShaderMap.GetReference()) >> 32), (int)((int64)(GameThreadShaderMap.GetReference())));
-#endif
 		OutstandingCompileShaderMapIds.AddUnique(GameThreadShaderMap->GetCompilingId());
 		UE_LOG(LogShaders, Log, TEXT("CacheShaders AddUniqueExisting %p %d"), this, GameThreadShaderMap->GetCompilingId());
 
diff --git a/Engine/Plugins/Experimental/NNERuntimeIREE/Source/NNERuntimeIREEShader/Private/NNERuntimeIREEShaderMap.cpp b/Engine/Plugins/Experimental/NNERuntimeIREE/Source/NNERuntimeIREEShader/Private/NNERuntimeIREEShaderMap.cpp
--- a/Engine/Plugins/Experimental/NNERuntimeIREE/Source/NNERuntimeIREEShader/Private/NNERuntimeIREEShaderMap.cpp
+++ b/Engine/Plugins/Experimental/NNERuntimeIREE/Source/NNERuntimeIREEShader/Private/NNERuntimeIREEShaderMap.cpp
@@ -322,9 +322,6 @@ void FNNERuntimeIREEShaderMap::Compile(
 				TArray<FNNERuntimeIREEResource*> NewCorrespondingKernels;
 				NewCorrespondingKernels.Add(InKernel);
 				ShaderMapsBeingCompiled.Add(this, NewCorrespondingKernels);
-#if DEBUG_INFINITESHADERCOMPILE
-				UE_LOG(LogTemp, Display, TEXT("Added NNERuntimeIREE ShaderMap 0x%08X%08X with kernel 0x%08X%08X to NNERuntimeIREEShaderMapsBeingCompiled"), (int)((int64)(this) >> 32), (int)((int64)(this)), (int)((int64)(InKernel) >> 32), (int)((int64)(InKernel)));
-#endif  
 			}
   
 			// Compile the shaders for this shader map now if not deferring and deferred compiles are not enabled globally
diff --git a/Engine/Plugins/FX/Niagara/Source/NiagaraShader/Private/NiagaraShader.cpp b/Engine/Plugins/FX/Niagara/Source/NiagaraShader/Private/NiagaraShader.cpp
--- a/Engine/Plugins/FX/Niagara/Source/NiagaraShader/Private/NiagaraShader.cpp
+++ b/Engine/Plugins/FX/Niagara/Source/NiagaraShader/Private/NiagaraShader.cpp
@@ -630,9 +630,6 @@ void FNiagaraShaderMap::Compile(
 			TArray<FNiagaraShaderScript*> NewCorrespondingScripts;
 			NewCorrespondingScripts.Add(Script);
 			NiagaraShaderMapsBeingCompiled.Add(this, NewCorrespondingScripts);
-#if DEBUG_INFINITESHADERCOMPILE
-			UE_LOG(LogTemp, Display, TEXT("Added Niagara ShaderMap 0x%08X%08X with Script 0x%08X%08X to NiagaraShaderMapsBeingCompiled"), (int)((int64)(this) >> 32), (int)((int64)(this)), (int)((int64)(Script) >> 32), (int)((int64)(Script)));
-#endif  
 			// Setup the compilation environment.
 			Script->SetupShaderCompilationEnvironment(InPlatform, *CompilationEnvironment);
   
@@ -858,9 +855,6 @@ bool FNiagaraShaderMap::TryToAddToExistingCompilationTask(FNiagaraShaderScript*
 
 		UE_LOG(LogShaders, Log, TEXT("TryToAddToExistingCompilationTask %p %d"), Script, GetCompilingId());
 
-#if DEBUG_INFINITESHADERCOMPILE
-		UE_LOG(LogTemp, Display, TEXT("Added shader map 0x%08X%08X from Niagara script 0x%08X%08X"), (int)((int64)(this) >> 32), (int)((int64)(this)), (int)((int64)(Script) >> 32), (int)((int64)(Script)));
-#endif
 		return true;
 	}
 
@@ -1096,12 +1090,6 @@ bool FNiagaraShaderMap::RemovePendingScript(FNiagaraShaderScript* Script)
 			}
 			bRemoved = true;
 		}
-#if DEBUG_INFINITESHADERCOMPILE
-		if ( Result )
-		{
-			UE_LOG(LogTemp, Display, TEXT("Removed shader map 0x%08X%08X from script 0x%08X%08X"), (int)((int64)(It.Key().GetReference()) >> 32), (int)((int64)(It.Key().GetReference())), (int)((int64)(Script) >> 32), (int)((int64)(Script)));
-		}
-#endif
 	}
 
 	return bRemoved;
diff --git a/Engine/Plugins/FX/Niagara/Source/NiagaraShader/Private/NiagaraShared.cpp b/Engine/Plugins/FX/Niagara/Source/NiagaraShader/Private/NiagaraShared.cpp
--- a/Engine/Plugins/FX/Niagara/Source/NiagaraShader/Private/NiagaraShared.cpp
+++ b/Engine/Plugins/FX/Niagara/Source/NiagaraShader/Private/NiagaraShared.cpp
@@ -632,9 +632,6 @@ bool FNiagaraShaderScript::CacheShaders(const FNiagaraShaderMapId& ShaderMapId,
 	if (GameThreadShaderMap && GameThreadShaderMap->TryToAddToExistingCompilationTask(this))
 	{
 		//FNiagaraShaderMap::ShaderMapsBeingCompiled.Find(GameThreadShaderMap);
-#if DEBUG_INFINITESHADERCOMPILE
-		UE_LOG(LogTemp, Display, TEXT("Found existing compiling shader for Niagara script %s, linking to other GameThreadShaderMap 0x%08X%08X"), *GetFriendlyName(), (int)((int64)(GameThreadShaderMap.GetReference()) >> 32), (int)((int64)(GameThreadShaderMap.GetReference())));
-#endif
 		OutstandingCompileShaderMapIds.AddUnique(GameThreadShaderMap->GetCompilingId());
 		UE_LOG(LogShaders, Log, TEXT("CacheShaders AddUniqueExisting %p %d"), this, GameThreadShaderMap->GetCompilingId());
 
diff --git a/Engine/Plugins/Runtime/ComputeFramework/Source/ComputeFramework/Private/ComputeKernelShaderMap.cpp b/Engine/Plugins/Runtime/ComputeFramework/Source/ComputeFramework/Private/ComputeKernelShaderMap.cpp
--- a/Engine/Plugins/Runtime/ComputeFramework/Source/ComputeFramework/Private/ComputeKernelShaderMap.cpp
+++ b/Engine/Plugins/Runtime/ComputeFramework/Source/ComputeFramework/Private/ComputeKernelShaderMap.cpp
@@ -416,9 +416,6 @@ void FComputeKernelShaderMap::Compile(
 				TArray<FComputeKernelResource*> NewCorrespondingKernels;
 				NewCorrespondingKernels.Add(InKernel);
 				ComputeKernelShaderMapsBeingCompiled.Add(this, NewCorrespondingKernels);
-#if DEBUG_INFINITESHADERCOMPILE
-				UE_LOG(LogTemp, Display, TEXT("Added ComputeKernel ShaderMap 0x%08X%08X with kernel 0x%08X%08X to ComputeKernelShaderMapsBeingCompiled"), (int)((int64)(this) >> 32), (int)((int64)(this)), (int)((int64)(InKernel) >> 32), (int)((int64)(InKernel)));
-#endif  
 			}
   
 			// Compile the shaders for this shader map now if not deferring and deferred compiles are not enabled globally
@@ -496,9 +493,6 @@ bool FComputeKernelShaderMap::TryToAddToExistingCompilationTask(FComputeKernelRe
 
 		UE_LOG(LogComputeFramework, Verbose, TEXT("TryToAddToExistingCompilationTask %p %d"), InKernel, GetCompilingId());
 
-#if DEBUG_INFINITESHADERCOMPILE
-		UE_LOG(LogTemp, Display, TEXT("Added shader map 0x%08X%08X from ComputeKernel 0x%08X%08X"), (int)((int64)(this) >> 32), (int)((int64)(this)), (int)((int64)(InKernel) >> 32), (int)((int64)(InKernel)));
-#endif
 		return true;
 	}
 
@@ -692,12 +686,6 @@ void FComputeKernelShaderMap::RemovePending(FComputeKernelResource* InKernel)
 			InKernel->RemoveOutstandingCompileId(It.Key()->CompilingId);
 			InKernel->NotifyCompilationFinished(FString::Printf(TEXT("%s: Removing compilation."), *InKernel->GetFriendlyName()));
 		}
-#if DEBUG_INFINITESHADERCOMPILE
-		if ( Result )
-		{
-			UE_LOG(LogTemp, Display, TEXT("Removed shader map 0x%08X%08X from kernel 0x%08X%08X"), (int)((int64)(It.Key().GetReference()) >> 32), (int)((int64)(It.Key().GetReference())), (int)((int64)(InKernel) >> 32), (int)((int64)(InKernel)));
-		}
-#endif
 	}
 }
 
diff --git a/Engine/Plugins/Runtime/ComputeFramework/Source/ComputeFramework/Private/ComputeKernelShared.cpp b/Engine/Plugins/Runtime/ComputeFramework/Source/ComputeFramework/Private/ComputeKernelShared.cpp
--- a/Engine/Plugins/Runtime/ComputeFramework/Source/ComputeFramework/Private/ComputeKernelShared.cpp
+++ b/Engine/Plugins/Runtime/ComputeFramework/Source/ComputeFramework/Private/ComputeKernelShared.cpp
@@ -363,9 +363,6 @@ bool FComputeKernelResource::CacheShaders(const FComputeKernelShaderMapId& InSha
 #if WITH_EDITOR
 	if (GameThreadShaderMap && GameThreadShaderMap->TryToAddToExistingCompilationTask(this))
 	{
-#if DEBUG_INFINITESHADERCOMPILE
-		UE_LOG(LogComputeFramework, Display, TEXT("Found existing compiling shader for kernel %s, linking to other GameThreadShaderMap 0x%08X%08X"), *GetFriendlyName(), (int)((int64)(GameThreadShaderMap.GetReference()) >> 32), (int)((int64)(GameThreadShaderMap.GetReference())));
-#endif
 		OutstandingCompileShaderMapIds.AddUnique(GameThreadShaderMap->GetCompilingId());
 		UE_LOG(LogComputeFramework, Verbose, TEXT("CacheShaders AddUniqueExisting %p %d"), this, GameThreadShaderMap->GetCompilingId());
 
diff --git a/Engine/Source/Runtime/Engine/Private/Materials/MaterialShader.cpp b/Engine/Source/Runtime/Engine/Private/Materials/MaterialShader.cpp
--- a/Engine/Source/Runtime/Engine/Private/Materials/MaterialShader.cpp
+++ b/Engine/Source/Runtime/Engine/Private/Materials/MaterialShader.cpp
@@ -2727,11 +2727,7 @@ void FMaterialShaderMap::Compile(
 	check(!Material->bContainsInlineShaders);
 
 	// Assign a unique identifier so that shaders from this shader map can be associated with it after a deferred compile
-	AcquireCompilingId(MaterialEnvironment);
-
-#if DEBUG_INFINITESHADERCOMPILE
-	UE_LOG(LogTemp, Display, TEXT("Added material ShaderMap 0x%08X%08X with Material 0x%08X%08X to ShaderMapsBeingCompiled"), (int)((int64)(this) >> 32), (int)((int64)(this)), (int)((int64)(Material) >> 32), (int)((int64)(Material)));
-#endif  
+	AcquireCompilingId(MaterialEnvironment); 
 
 	FMaterialShaderMapContent* NewContent = new FMaterialShaderMapContent(InPlatform);
 	NewContent->MaterialCompilationOutput = InMaterialCompilationOutput;
diff --git a/Engine/Source/Runtime/Engine/Private/Materials/MaterialShared.cpp b/Engine/Source/Runtime/Engine/Private/Materials/MaterialShared.cpp
--- a/Engine/Source/Runtime/Engine/Private/Materials/MaterialShared.cpp
+++ b/Engine/Source/Runtime/Engine/Private/Materials/MaterialShared.cpp
@@ -3583,9 +3583,6 @@ bool FMaterial::BeginCompileShaderMap(
 		}
 		else
 		{
-#if DEBUG_INFINITESHADERCOMPILE
-			UE_LOG(LogTemp, Display, TEXT("Kicking off shader compilation for %s, GameThreadShaderMap 0x%08X%08X"), *GetFriendlyName(), (int)((int64)(NewShaderMap.GetReference()) >> 32), (int)((int64)(NewShaderMap.GetReference())));
-#endif
 			SetGameThreadShaderMap(NewShaderMap->AcquireFinalizedClone());
 		}
 	}
diff --git a/Engine/Source/Runtime/Engine/Private/ShaderCompiler/ShaderCompiler.cpp b/Engine/Source/Runtime/Engine/Private/ShaderCompiler/ShaderCompiler.cpp
--- a/Engine/Source/Runtime/Engine/Private/ShaderCompiler/ShaderCompiler.cpp
+++ b/Engine/Source/Runtime/Engine/Private/ShaderCompiler/ShaderCompiler.cpp
@@ -2086,9 +2086,6 @@ void FShaderCompilingManager::ProcessCompiledShaderMaps(
 
 			FMaterialShaderMap* ShaderMapToUseForRendering = nullptr;
 
-#if DEBUG_INFINITESHADERCOMPILE
-			UE_LOG(LogTemp, Display, TEXT("Finished compile of shader map 0x%08X%08X"), (int)((int64)(ShaderMap.GetReference()) >> 32), (int)((int64)(ShaderMap.GetReference())));
-#endif
 			int32 NumIncompleteMaterials = 0;
 			int32 MaterialIndex = 0;
 			
@@ -2101,10 +2098,6 @@ void FShaderCompilingManager::ProcessCompiledShaderMaps(
 				check(Material->GetGameThreadCompilingShaderMapId() == CompilingShaderMap->GetCompilingId());
 				bRequiredComplete |= Material->IsRequiredComplete();
 
-#if DEBUG_INFINITESHADERCOMPILE
-				UE_LOG(LogTemp, Display, TEXT("Shader map %s complete, GameThreadShaderMap 0x%08X%08X, marking material %s as finished"), *ShaderMap->GetFriendlyName(), (int)((int64)(ShaderMap.GetReference()) >> 32), (int)((int64)(ShaderMap.GetReference())), *Material->GetFriendlyName());
-				UE_LOG(LogTemp, Display, TEXT("Marking material as finished 0x%08X%08X"), (int)((int64)(Material) >> 32), (int)((int64)(Material)));
-#endif
 				//Material->RemoveOutstandingCompileId(ShaderMap->CompilingId);
 
 				bool bReleaseCompilingId = false;
diff --git a/Engine/Source/Runtime/Engine/Public/ShaderCompiler.h b/Engine/Source/Runtime/Engine/Public/ShaderCompiler.h
--- a/Engine/Source/Runtime/Engine/Public/ShaderCompiler.h
+++ b/Engine/Source/Runtime/Engine/Public/ShaderCompiler.h
@@ -48,8 +48,6 @@ struct FAnalyticsEventAttribute;
 
 DECLARE_LOG_CATEGORY_EXTERN(LogShaderCompilers, Log, All);
 
-#define DEBUG_INFINITESHADERCOMPILE 0
-
 bool AreShaderErrorsFatal();
 bool CreateShadersOnLoad();
 
