OpenColorIO: Restrain precise keyword to when REFACTORING_ALLOWED is defined, per specification.

#rnx

diff --git a/Engine/Plugins/Compositing/OpenColorIO/Shaders/Private/OpenColorIOShader.usf b/Engine/Plugins/Compositing/OpenColorIO/Shaders/Private/OpenColorIOShader.usf
--- a/Engine/Plugins/Compositing/OpenColorIO/Shaders/Private/OpenColorIOShader.usf
+++ b/Engine/Plugins/Compositing/OpenColorIO/Shaders/Private/OpenColorIOShader.usf
@@ -6,6 +6,10 @@
 
 #include "/Engine/Public/Platform.ush"
 
+#ifndef REFACTORING_ALLOWED
+#define precise // Empty definition when REFACTORING_ALLOWED flag is not defined, per specification: https://learn.microsoft.com/en-us/windows/win32/direct3dhlsl/precise
+#endif
+
 /* Declaration of all variables
 =============================================================================*/
 Texture2D		InputTexture;
