OpenColorIO: Cleanup of deprecated functions & members.

diff --git a/Engine/Plugins/Compositing/OpenColorIO/Source/OpenColorIO/Private/OpenColorIOColorTransform.cpp b/Engine/Plugins/Compositing/OpenColorIO/Source/OpenColorIO/Private/OpenColorIOColorTransform.cpp
--- a/Engine/Plugins/Compositing/OpenColorIO/Source/OpenColorIO/Private/OpenColorIOColorTransform.cpp
+++ b/Engine/Plugins/Compositing/OpenColorIO/Source/OpenColorIO/Private/OpenColorIOColorTransform.cpp
@@ -496,23 +496,6 @@ bool UOpenColorIOColorTransform::GetRenderResources(ERHIFeatureLevel::Type InFea
 	return OutShaderResource != nullptr;
 }
 
-// Deprecated
-bool UOpenColorIOColorTransform::GetRenderResources(ERHIFeatureLevel::Type InFeatureLevel, FOpenColorIOTransformResource*& OutShaderResource, TSortedMap<int32, TWeakObjectPtr<UTexture>>& OutTextureResources) const
-{
-	TArray<TWeakObjectPtr<UTexture>> OutTextureArray;
-	if (GetRenderResources(InFeatureLevel, OutShaderResource, OutTextureArray))
-	{
-		for (int32 Index = 0; Index < OutTextureArray.Num(); Index++)
-		{
-			OutTextureResources.Emplace(Index, OutTextureArray[Index]);
-		}
-
-		return true;
-	}
-
-	return false;
-}
-
 bool UOpenColorIOColorTransform::AreRenderResourcesReady() const
 {
 	check(IsInGameThread());
diff --git a/Engine/Plugins/Compositing/OpenColorIO/Source/OpenColorIO/Private/OpenColorIOConfiguration.cpp b/Engine/Plugins/Compositing/OpenColorIO/Source/OpenColorIO/Private/OpenColorIOConfiguration.cpp
--- a/Engine/Plugins/Compositing/OpenColorIO/Source/OpenColorIO/Private/OpenColorIOConfiguration.cpp
+++ b/Engine/Plugins/Compositing/OpenColorIO/Source/OpenColorIO/Private/OpenColorIOConfiguration.cpp
@@ -83,19 +83,6 @@ bool UOpenColorIOConfiguration::GetRenderResources(ERHIFeatureLevel::Type InFeat
 	return false;
 }
 
-bool UOpenColorIOConfiguration::GetRenderResources(ERHIFeatureLevel::Type InFeatureLevel, const FOpenColorIOColorConversionSettings& InSettings, FOpenColorIOTransformResource*& OutShaderResource, TSortedMap<int32, TWeakObjectPtr<UTexture>>& OutTextureResources)
-{
-	TObjectPtr<const UOpenColorIOColorTransform> Transform = FindTransform(InSettings);
-	if (IsValid(Transform))
-	{
-		PRAGMA_DISABLE_DEPRECATION_WARNINGS
-		return Transform->GetRenderResources(InFeatureLevel, OutShaderResource, OutTextureResources);
-		PRAGMA_ENABLE_DEPRECATION_WARNINGS
-	}
-
-	return false;
-}
-
 bool UOpenColorIOConfiguration::HasTransform(const FString& InSourceColorSpace, const FString& InDestinationColorSpace)
 {
 	TObjectPtr<UOpenColorIOColorTransform>* TransformData = ColorTransforms.FindByPredicate([&](const UOpenColorIOColorTransform* InTransformData)
diff --git a/Engine/Plugins/Compositing/OpenColorIO/Source/OpenColorIO/Private/OpenColorIORendering.cpp b/Engine/Plugins/Compositing/OpenColorIO/Source/OpenColorIO/Private/OpenColorIORendering.cpp
--- a/Engine/Plugins/Compositing/OpenColorIO/Source/OpenColorIO/Private/OpenColorIORendering.cpp
+++ b/Engine/Plugins/Compositing/OpenColorIO/Source/OpenColorIO/Private/OpenColorIORendering.cpp
@@ -22,9 +22,6 @@
 #include "ShaderParameterStruct.h"
 #include "TextureResource.h"
 
-
-float FOpenColorIORendering::DefaultDisplayGamma = 2.2f;
-
 // static
 void FOpenColorIORendering::AddPass_RenderThread(
 	FRDGBuilder& GraphBuilder,
diff --git a/Engine/Plugins/Compositing/OpenColorIO/Source/OpenColorIO/Private/OpenColorIOShader.cpp b/Engine/Plugins/Compositing/OpenColorIO/Source/OpenColorIO/Private/OpenColorIOShader.cpp
--- a/Engine/Plugins/Compositing/OpenColorIO/Source/OpenColorIO/Private/OpenColorIOShader.cpp
+++ b/Engine/Plugins/Compositing/OpenColorIO/Source/OpenColorIO/Private/OpenColorIOShader.cpp
@@ -5,11 +5,6 @@
 #include "OpenColorIOModule.h"
 #include "SystemTextures.h"
 
-bool OpenColorIOBindTextureResources(FOpenColorIOPixelShaderParameters* Parameters, const TSortedMap<int32, FTextureResource*>& InTextureResources)
-{
-	return false;
-}
-
 FOpenColorIOPixelShader::FOpenColorIOPixelShader(const FOpenColorIOShaderType::CompiledShaderInitializerType& Initializer)
 	: FGlobalShader(Initializer)
 {
diff --git a/Engine/Plugins/Compositing/OpenColorIO/Source/OpenColorIO/Public/OpenColorIOColorTransform.h b/Engine/Plugins/Compositing/OpenColorIO/Source/OpenColorIO/Public/OpenColorIOColorTransform.h
--- a/Engine/Plugins/Compositing/OpenColorIO/Source/OpenColorIO/Public/OpenColorIOColorTransform.h
+++ b/Engine/Plugins/Compositing/OpenColorIO/Source/OpenColorIO/Public/OpenColorIOColorTransform.h
@@ -51,9 +51,6 @@ public:
 	 */
 	UE_API bool GetRenderResources(ERHIFeatureLevel::Type InFeatureLevel, FOpenColorIOTransformResource*& OutShaderResource, TArray<TWeakObjectPtr<UTexture>>& OutTextureResources) const;
 
-	UE_DEPRECATED(5.6, "This method is deprecated, please use the one with an array of textures instead.")
-	UE_API bool GetRenderResources(ERHIFeatureLevel::Type InFeatureLevel, FOpenColorIOTransformResource*& OutShaderResource, TSortedMap<int32, TWeakObjectPtr<UTexture>>& OutTextureResources) const;
-
 	/**
 	 * Returns true if shader/texture resources have finished compiling and are ready for use (to be called on the game thread).
 	 */
diff --git a/Engine/Plugins/Compositing/OpenColorIO/Source/OpenColorIO/Public/OpenColorIOConfiguration.h b/Engine/Plugins/Compositing/OpenColorIO/Source/OpenColorIO/Public/OpenColorIOConfiguration.h
--- a/Engine/Plugins/Compositing/OpenColorIO/Source/OpenColorIO/Public/OpenColorIOConfiguration.h
+++ b/Engine/Plugins/Compositing/OpenColorIO/Source/OpenColorIO/Public/OpenColorIOConfiguration.h
@@ -41,9 +41,6 @@ public:
 
 	/** Get the shader and optionally lookup texture resources to be used by the color transform render pass specified by the settings parameter. */
 	UE_API bool GetRenderResources(ERHIFeatureLevel::Type InFeatureLevel, const FOpenColorIOColorConversionSettings& InSettings, FOpenColorIOTransformResource*& OutShaderResource, TArray<TWeakObjectPtr<UTexture>>& OutTextureResources);
-
-	UE_DEPRECATED(5.6, "This method is deprecated, please use the one with an array of textures instead.")
-	UE_API bool GetRenderResources(ERHIFeatureLevel::Type InFeatureLevel, const FOpenColorIOColorConversionSettings& InSettings, FOpenColorIOTransformResource*& OutShaderResource, TSortedMap<int32, TWeakObjectPtr<UTexture>>& OutTextureResources);
 	
 	/** Check if this configuration contains a transform from the specified source & destination color spaces. */
 	UE_API bool HasTransform(const FString& InSourceColorSpace, const FString& InDestinationColorSpace);
diff --git a/Engine/Plugins/Compositing/OpenColorIO/Source/OpenColorIO/Public/OpenColorIORendering.h b/Engine/Plugins/Compositing/OpenColorIO/Source/OpenColorIO/Public/OpenColorIORendering.h
--- a/Engine/Plugins/Compositing/OpenColorIO/Source/OpenColorIO/Public/OpenColorIORendering.h
+++ b/Engine/Plugins/Compositing/OpenColorIO/Source/OpenColorIO/Public/OpenColorIORendering.h
@@ -116,9 +116,6 @@ public:
 	* This function sets up provided view and related family to properly support OCIO. 
 	*/
 	static UE_API void PrepareView(FSceneViewFamily& InViewFamily, FSceneView& InView);
-	
-	UE_DEPRECATED(5.6, "This constant is deprecated, please use UTextureRenderTarget::GetDefaultDisplayGamma() instead.")
-	static UE_API float DefaultDisplayGamma;
 };
 
 #undef UE_API
diff --git a/Engine/Plugins/Compositing/OpenColorIO/Source/OpenColorIO/Public/OpenColorIOShader.h b/Engine/Plugins/Compositing/OpenColorIO/Source/OpenColorIO/Public/OpenColorIOShader.h
--- a/Engine/Plugins/Compositing/OpenColorIO/Source/OpenColorIO/Public/OpenColorIOShader.h
+++ b/Engine/Plugins/Compositing/OpenColorIO/Source/OpenColorIO/Public/OpenColorIOShader.h
@@ -26,12 +26,6 @@ BEGIN_SHADER_PARAMETER_STRUCT(FOpenColorIOPixelShaderParameters, )
 	RENDER_TARGET_BINDING_SLOTS()
 END_SHADER_PARAMETER_STRUCT()
 
-namespace OpenColorIOShader
-{
-	UE_DEPRECATED(5.6, "This limit is now deprecated.")
-	static constexpr uint32 MaximumTextureSlots = 8;
-}
-
 class FOpenColorIOPixelShader : public FGlobalShader
 {
 	DECLARE_EXPORTED_SHADER_TYPE(FOpenColorIOPixelShader, OpenColorIO, OPENCOLORIO_API);
@@ -63,7 +57,4 @@ public:
 	using FParameters = FOpenColorIOInvalidShaderParameters;
 };
 
-UE_DEPRECATED(5.6, "This shader parameter binding function is now deprecated.")
-OPENCOLORIO_API bool OpenColorIOBindTextureResources(FOpenColorIOPixelShaderParameters* Parameters, const TSortedMap<int32, FTextureResource*>& InTextureResources);
-
 FRHITexture* OpenColorIOGetMiniFontTexture();
diff --git a/Engine/Plugins/Compositing/OpenColorIO/Source/OpenColorIO/Public/OpenColorIOShared.h b/Engine/Plugins/Compositing/OpenColorIO/Source/OpenColorIO/Public/OpenColorIOShared.h
--- a/Engine/Plugins/Compositing/OpenColorIO/Source/OpenColorIO/Public/OpenColorIOShared.h
+++ b/Engine/Plugins/Compositing/OpenColorIO/Source/OpenColorIO/Public/OpenColorIOShared.h
@@ -454,16 +454,6 @@ public:
 
 	const FString& GetFriendlyName()	const { return FriendlyName; }
 
-	UE_DEPRECATED(5.6, "This method is deprecated, please use the one with a shader parameter metadata instead.")
-	void SetupResource(
-		ERHIFeatureLevel::Type InFeatureLevel,
-		const FString& InShaderCodeHash,
-		const FString& InShadercode,
-		const FString& InRawConfigHash,
-		const FString& InFriendlyName,
-		const FName& InAssetPath
-	) { }
-
 	// Note: Takes ownership of the shader parameter metadata.
 	OPENCOLORIO_API void SetupResource(
 		ERHIFeatureLevel::Type InFeatureLevel,
