OCIO: Migrate to the UE5 DDC API

diff --git a/Engine/Plugins/Compositing/OpenColorIO/Source/OpenColorIO/OpenColorIO.Build.cs b/Engine/Plugins/Compositing/OpenColorIO/Source/OpenColorIO/OpenColorIO.Build.cs
--- a/Engine/Plugins/Compositing/OpenColorIO/Source/OpenColorIO/OpenColorIO.Build.cs
+++ b/Engine/Plugins/Compositing/OpenColorIO/Source/OpenColorIO/OpenColorIO.Build.cs
@@ -27,11 +27,6 @@ namespace UnrealBuildTool.Rules
 					"SlateCore",
 				});
 
-			PrivateIncludePathModuleNames.AddRange(
-				new string[] {
-					"DerivedDataCache"
-				});
-
 			PublicDependencyModuleNames.AddRange(
 				new string[] {
 					"Core",
@@ -41,6 +36,11 @@ namespace UnrealBuildTool.Rules
 					"OpenColorIOWrapper"
 				});
 
+			if (Target.bBuildWithEditorOnlyData)
+			{
+				PrivateDependencyModuleNames.Add("DerivedDataCache");
+			}
+
 			if (Target.bBuildEditor == true)
 			{
 				PrivateDependencyModuleNames.AddRange(
diff --git a/Engine/Plugins/Compositing/OpenColorIO/Source/OpenColorIO/Private/OpenColorIOShaderMap.cpp b/Engine/Plugins/Compositing/OpenColorIO/Source/OpenColorIO/Private/OpenColorIOShaderMap.cpp
--- a/Engine/Plugins/Compositing/OpenColorIO/Source/OpenColorIO/Private/OpenColorIOShaderMap.cpp
+++ b/Engine/Plugins/Compositing/OpenColorIO/Source/OpenColorIO/Private/OpenColorIOShaderMap.cpp
@@ -17,7 +17,8 @@
 
 #if WITH_EDITOR
 	#include "OpenColorIOWrapper.h"
-	#include "DerivedDataCacheInterface.h"
+	#include "DerivedDataCache.h"
+	#include "DerivedDataRequestOwner.h"
 	#include "Interfaces/ITargetPlatformManagerModule.h"
 	#include "TickableEditorObject.h"
 #endif
@@ -250,14 +251,14 @@ void OpenColorIOShaderMapAppendKeyString(EShaderPlatform InPlatform, FString& Ou
 }
 
 /** Creates a string key for the derived data cache given a shader map id. */
-static FString GetOpenColorIOShaderMapKeyString(const FOpenColorIOShaderMapId& InShaderMapId, EShaderPlatform InPlatform)
+static UE::DerivedData::FCacheKey GetOpenColorIOShaderMapKeyString(const FOpenColorIOShaderMapId& InShaderMapId, EShaderPlatform InPlatform)
 {
 	const FName Format = LegacyShaderPlatformToShaderFormat(InPlatform);
 	FString ShaderMapKeyString = Format.ToString() + TEXT("_") + FString(FString::FromInt(GetTargetPlatformManagerRef().ShaderFormatVersion(Format))) + TEXT("_");
 	OpenColorIOShaderMapAppendKeyString(InPlatform, ShaderMapKeyString);
 	ShaderMapAppendKeyString(InPlatform, ShaderMapKeyString);
 	InShaderMapId.AppendKeyString(ShaderMapKeyString);
-	return FDerivedDataCacheInterface::BuildCacheKey(TEXT("OCIOSM"), OPENCOLORIO_DERIVEDDATA_VER, *ShaderMapKeyString);
+	return UE::DerivedData::BuildLegacyCacheKey(TEXT("OCIOSM"), OPENCOLORIO_DERIVEDDATA_VER, ShaderMapKeyString);
 }
 
 void FOpenColorIOShaderMap::LoadFromDerivedDataCache(const FOpenColorIOTransformResource* InColorTransform, const FOpenColorIOShaderMapId& InShaderMapId, EShaderPlatform InPlatform, TRefCountPtr<FOpenColorIOShaderMap>& InOutShaderMap)
@@ -281,29 +282,36 @@ void FOpenColorIOShaderMap::LoadFromDerivedDataCache(const FOpenColorIOTransform
 			SCOPE_SECONDS_COUNTER(OpenColorIOShaderDDCTime);
 			COOK_STAT(auto Timer = OpenColorIOShaderCookStats::UsageStats.TimeSyncWork());
 
-			TArray<uint8> CachedData;
-			const FString DataKey = GetOpenColorIOShaderMapKeyString(InShaderMapId, InPlatform);
-
-			if (GetDerivedDataCacheRef().GetSynchronous(*DataKey, CachedData, InColorTransform->GetFriendlyName()))
+			using namespace UE::DerivedData;
+			const FCacheKey Key = GetOpenColorIOShaderMapKeyString(InShaderMapId, InPlatform);
+			FRequestOwner Owner(EPriority::Blocking);
+			GetCache().GetValue({{{InColorTransform->GetFriendlyName()}, Key}}, Owner, [InShaderMapId, InPlatform, &InOutShaderMap, &Timer](FCacheGetValueResponse&& Response)
 			{
-				COOK_STAT(Timer.AddHit(CachedData.Num()));
-				InOutShaderMap = new FOpenColorIOShaderMap();
-				FMemoryReader Ar(CachedData, true);
+				if (Response.Status == EStatus::Ok)
+				{
+					COOK_STAT(Timer.AddHit(Response.Value.GetRawSize()));
 
-				// Deserialize from the cached data
-				InOutShaderMap->Serialize(Ar);
+					FSharedBuffer CachedData = Response.Value.GetData().Decompress();
 
-				checkSlow(InOutShaderMap->GetShaderMapId() == InShaderMapId);
+					InOutShaderMap = new FOpenColorIOShaderMap();
+					FMemoryReaderView Ar(CachedData, true);
 
-				// Register in the global map
-				InOutShaderMap->Register(InPlatform);
-			}
-			else
-			{
-				// We should be build the data later, and we can track that the resource was built there when we push it to the DDC.
-				COOK_STAT(Timer.TrackCyclesOnly());
-				InOutShaderMap = nullptr;
-			}
+					// Deserialize from the cached data
+					InOutShaderMap->Serialize(Ar);
+
+					checkSlow(InOutShaderMap->GetShaderMapId() == InShaderMapId);
+
+					// Register in the global map
+					InOutShaderMap->Register(InPlatform);
+				}
+				else
+				{
+					// We should be build the data later, and we can track that the resource was built there when we push it to the DDC.
+					COOK_STAT(Timer.TrackCyclesOnly());
+					InOutShaderMap = nullptr;
+				}
+			});
+			Owner.Wait();
 		}
 		INC_FLOAT_STAT_BY(STAT_ShaderCompiling_DDCLoading,(float)OpenColorIOShaderDDCTime);
 	}
@@ -311,13 +319,22 @@ void FOpenColorIOShaderMap::LoadFromDerivedDataCache(const FOpenColorIOTransform
 
 void FOpenColorIOShaderMap::SaveToDerivedDataCache()
 {
+	using namespace UE::DerivedData;
+
 	COOK_STAT(auto Timer = OpenColorIOShaderCookStats::UsageStats.TimeSyncWork());
-	TArray<uint8> SaveData;
-	FMemoryWriter Ar(SaveData, true);
-	Serialize(Ar);
+	FValue Value;
+	{
+		TArray<uint8> SaveData;
+		FMemoryWriter Ar(SaveData, true);
+		Serialize(Ar);
+		Value = FValue::Compress(MakeSharedBufferFromArray(MoveTemp(SaveData)));
+	}
 
-	GetDerivedDataCacheRef().Put(*GetOpenColorIOShaderMapKeyString(GetContent()->ShaderMapId, GetShaderPlatform()), SaveData, FStringView(*GetFriendlyName()));
-	COOK_STAT(Timer.AddMiss(SaveData.Num()));
+	const FCacheKey Key = GetOpenColorIOShaderMapKeyString(GetContent()->ShaderMapId, GetShaderPlatform());
+	FRequestOwner Owner(EPriority::Normal);
+	GetCache().PutValue({{{GetFriendlyName()}, Key, Value}}, Owner);
+	Owner.KeepAlive();
+	COOK_STAT(Timer.AddMiss(Value.GetRawSize()));
 }
 
 /**
