OCIO: Automatically replace engine paths with the {Engine} token.

diff --git a/Engine/Plugins/Compositing/OpenColorIO/Source/OpenColorIO/Private/OpenColorIOColorTransform.cpp b/Engine/Plugins/Compositing/OpenColorIO/Source/OpenColorIO/Private/OpenColorIOColorTransform.cpp
--- a/Engine/Plugins/Compositing/OpenColorIO/Source/OpenColorIO/Private/OpenColorIOColorTransform.cpp
+++ b/Engine/Plugins/Compositing/OpenColorIO/Source/OpenColorIO/Private/OpenColorIOColorTransform.cpp
@@ -569,7 +569,7 @@ bool UOpenColorIOColorTransform::GetTransformProcessor(FOpenColorIOWrapperProces
 
 	if (!ConfigWrapper)
 	{
-		UE_LOG(LogOpenColorIO, Error, TEXT("Failed to create processor for color transform %s. Configuration file [%s] was invalid."), *GetTransformFriendlyName(), ConfigurationOwner ? *ConfigurationOwner->ConfigurationFile.FilePath : TEXT("Missing"));
+		UE_LOG(LogOpenColorIO, Error, TEXT("Failed to create processor for color transform %s. Configuration file [%s] was invalid."), *GetTransformFriendlyName(), ConfigurationOwner ? *ConfigurationOwner->GetConfigurationFile().FilePath : TEXT("Missing"));
 		return false;
 	}
 
@@ -865,7 +865,7 @@ bool UOpenColorIOColorTransform::UpdateShaderInfo(FString& OutShaderCodeHash, FS
 	}
 	else
 	{
-		UE_LOG(LogOpenColorIO, Error, TEXT("Failed to fetch shader info for color transform %s. Configuration file [%s] was invalid."), *GetTransformFriendlyName(), ConfigurationOwner ? *ConfigurationOwner->ConfigurationFile.FilePath : TEXT("Missing"));
+		UE_LOG(LogOpenColorIO, Error, TEXT("Failed to fetch shader info for color transform %s. Configuration file [%s] was invalid."), *GetTransformFriendlyName(), ConfigurationOwner ? *ConfigurationOwner->GetConfigurationFile().FilePath : TEXT("Missing"));
 	}
 
 	return false;
diff --git a/Engine/Plugins/Compositing/OpenColorIO/Source/OpenColorIO/Private/OpenColorIOConfiguration.cpp b/Engine/Plugins/Compositing/OpenColorIO/Source/OpenColorIO/Private/OpenColorIOConfiguration.cpp
--- a/Engine/Plugins/Compositing/OpenColorIO/Source/OpenColorIO/Private/OpenColorIOConfiguration.cpp
+++ b/Engine/Plugins/Compositing/OpenColorIO/Source/OpenColorIO/Private/OpenColorIOConfiguration.cpp
@@ -706,7 +706,7 @@ void UOpenColorIOConfiguration::PostEditChangeProperty(FPropertyChangedEvent& Pr
 	if (PropertyChangedEvent.MemberProperty && PropertyChangedEvent.MemberProperty->GetFName() == GET_MEMBER_NAME_CHECKED(UOpenColorIOConfiguration, ConfigurationFile))
 	{
 		// Note: Reload calls LoadConfiguration() internally.
-		ReloadExistingColorspaces();
+		SetConfigurationFile(ConfigurationFile);
 	}
 	else if (PropertyChangedEvent.GetPropertyName() == GET_MEMBER_NAME_CHECKED(UOpenColorIOConfiguration, DesiredColorSpaces))
 	{
@@ -790,6 +790,21 @@ void UOpenColorIOConfiguration::OnToastCallback(bool bInReloadColorspaces)
 	}
 }
 
+void UOpenColorIOConfiguration::SetConfigurationFile(const FFilePath& InFilePath)
+{
+	ConfigurationFile = InFilePath;
+
+	const FString EngineFullPath = FPaths::ConvertRelativePathToFull(FPaths::EngineDir());
+
+	if (!EngineFullPath.IsEmpty() && ConfigurationFile.FilePath.Contains(EngineFullPath))
+	{
+		// Automatically replace the engine path with its {Engine} token for better compatilibilty across different machines.
+		ConfigurationFile.FilePath = ConfigurationFile.FilePath.Replace(*EngineFullPath, TEXT("{Engine}/"));
+	}
+
+	ReloadExistingColorspaces();
+}
+
 #endif //WITH_EDITOR
 
 
diff --git a/Engine/Plugins/Compositing/OpenColorIO/Source/OpenColorIO/Public/OpenColorIOConfiguration.h b/Engine/Plugins/Compositing/OpenColorIO/Source/OpenColorIO/Public/OpenColorIOConfiguration.h
--- a/Engine/Plugins/Compositing/OpenColorIO/Source/OpenColorIO/Public/OpenColorIOConfiguration.h
+++ b/Engine/Plugins/Compositing/OpenColorIO/Source/OpenColorIO/Public/OpenColorIOConfiguration.h
@@ -136,12 +136,19 @@ private:
 
 public:
 
-	UPROPERTY(BlueprintReadWrite, EditAnywhere, Category = "Config", meta = (FilePathFilter = "Config Files (*.ocio, *.ocioz)|*.ocio;*.ocioz", RelativeToGameDir))
-	FFilePath ConfigurationFile;
+	/** Get the configuration file path. */
+	UFUNCTION(BlueprintGetter)
+	UE_API const FFilePath& GetConfigurationFile() const { return ConfigurationFile; }
+
+	/** Set the configuration file path. */
+	UFUNCTION(BlueprintSetter)
+	UE_API void SetConfigurationFile(const FFilePath& InFilePath);
 
+	/** Desired color space tranforms. */
 	UPROPERTY(BlueprintReadWrite, EditAnywhere, Category = "Transform")
 	TArray<FOpenColorIOColorSpace> DesiredColorSpaces;
 
+	/** Desired display view tranforms. */
 	UPROPERTY(BlueprintReadWrite, EditAnywhere, Category = "Transform", DisplayName="Desired Display-Views")
 	TArray<FOpenColorIODisplayView> DesiredDisplayViews;
 
@@ -152,6 +159,9 @@ public:
 	TMap<FString, FString> Context;
 
 private:
+	/** Configuration file path. */
+	UPROPERTY(EditAnywhere, BlueprintGetter = GetConfigurationFile, BlueprintSetter = SetConfigurationFile, Category = "Config", meta = (AllowPrivateAccess = true, FilePathFilter = "Config Files (*.ocio, *.ocioz)|*.ocio;*.ocioz", RelativeToGameDir))
+	FFilePath ConfigurationFile;
 
 	UPROPERTY()
 	TArray<TObjectPtr<UOpenColorIOColorTransform>> ColorTransforms;
@@ -178,6 +188,9 @@ private:
 	/** Hash of all of the config content, including relevant external file information. */
 	UPROPERTY()
 	FString ConfigHash;
+
+	friend class UAssetDefinition_OpenColorIOConfiguration;
+	friend class FOpenColorIOConfigurationCustomization;
 };
 
 #undef UE_API
diff --git a/Engine/Plugins/Compositing/OpenColorIO/Source/OpenColorIOEditor/Private/Factories/OpenColorIOConfigurationFactory.cpp b/Engine/Plugins/Compositing/OpenColorIO/Source/OpenColorIOEditor/Private/Factories/OpenColorIOConfigurationFactory.cpp
--- a/Engine/Plugins/Compositing/OpenColorIO/Source/OpenColorIOEditor/Private/Factories/OpenColorIOConfigurationFactory.cpp
+++ b/Engine/Plugins/Compositing/OpenColorIO/Source/OpenColorIOEditor/Private/Factories/OpenColorIOConfigurationFactory.cpp
@@ -35,8 +35,7 @@ UObject* UOpenColorIOConfigurationFactory::FactoryCreateFile(UClass* InClass, UO
 		FString ConfigFilename = CurrentFilename;
 		FPaths::MakePathRelativeTo(ConfigFilename, *FPaths::ProjectDir());
 
-		Asset->ConfigurationFile = { ConfigFilename };
-		Asset->ReloadExistingColorspaces();
+		Asset->SetConfigurationFile(FFilePath{ ConfigFilename });
 
 		GEditor->GetEditorSubsystem<UImportSubsystem>()->BroadcastAssetPostImport(this, Asset);
 	}
