OCIO: Fix map iterator deletion CIS warning.

#Synced-CL 48809419

diff --git a/Engine/Plugins/Compositing/OpenColorIO/Source/OpenColorIO/Private/OpenColorIOShaderCompilationManager.cpp b/Engine/Plugins/Compositing/OpenColorIO/Source/OpenColorIO/Private/OpenColorIOShaderCompilationManager.cpp
--- a/Engine/Plugins/Compositing/OpenColorIO/Source/OpenColorIO/Private/OpenColorIOShaderCompilationManager.cpp
+++ b/Engine/Plugins/Compositing/OpenColorIO/Source/OpenColorIO/Private/OpenColorIOShaderCompilationManager.cpp
@@ -286,7 +286,7 @@ void FOpenColorIOShaderCompilationManager::ProcessCompiledOpenColorIOShaderMaps(
 
 				// Cleanup shader jobs and compile tracking structures
 				ResultArray.Empty();
-				CompiledShaderMaps.Remove(ShaderMap->GetCompilingId());
+				ProcessIt.RemoveCurrent();
 			}
 
 			if (TimeBudget < 0)
