OpenColorIO: Automatic migration of default config assets to avoid ACES 2.0 breakage.

#Synced-CL 49940565

diff --git a/Engine/Plugins/Compositing/OpenColorIO/Source/OpenColorIO/Private/OpenColorIOConfiguration.cpp b/Engine/Plugins/Compositing/OpenColorIO/Source/OpenColorIO/Private/OpenColorIOConfiguration.cpp
--- a/Engine/Plugins/Compositing/OpenColorIO/Source/OpenColorIO/Private/OpenColorIOConfiguration.cpp
+++ b/Engine/Plugins/Compositing/OpenColorIO/Source/OpenColorIO/Private/OpenColorIOConfiguration.cpp
@@ -19,6 +19,7 @@
 #include "Widgets/Notifications/SNotificationList.h"
 
 #include "UObject/AssetRegistryTagsContext.h"
+#include "UObject/FortniteMainBranchObjectVersion.h"
 #include "UObject/ObjectSaveContext.h"
 
 #include UE_INLINE_GENERATED_CPP_BY_NAME(OpenColorIOConfiguration)
@@ -240,7 +241,7 @@ void UOpenColorIOConfiguration::ReloadExistingColorspaces(bool bForce)
 
 			const UOpenColorIOSettings* Settings = GetDefault<UOpenColorIOSettings>();
 
-			// Genereate new shaders.
+			// Generate new shaders.
 			for (int32 indexTop = 0; indexTop < DesiredColorSpaces.Num(); ++indexTop)
 			{
 				const FOpenColorIOColorSpace& TopColorSpace = DesiredColorSpaces[indexTop];
@@ -523,11 +524,32 @@ void UOpenColorIOConfiguration::PostInitProperties()
 #endif //WITH_EDITOR
 }
 
+void UOpenColorIOConfiguration::Serialize(FArchive& Ar)
+{
+	Super::Serialize(Ar);
+
+	Ar.UsingCustomVersion(FFortniteMainBranchObjectVersion::GUID);
+}
+
 void UOpenColorIOConfiguration::PostLoad()
 {
 	Super::PostLoad();
 
 #if WITH_EDITOR
+	if (GetLinkerCustomVersion(FFortniteMainBranchObjectVersion::GUID) < FFortniteMainBranchObjectVersion::OpenColorIOv251)
+	{
+		// Since ACES 2.0 transforms are significant breaking change, we ease the transition by keeping old configs at the previous version.
+		// TODO: Refactor configuration assets so that the external config is no longer "live" and instead only used on import.
+		if (ConfigurationFile.FilePath == TEXT("ocio://default") || ConfigurationFile.FilePath == TEXT("ocio://cg-config-latest"))
+		{
+			ConfigurationFile.FilePath = TEXT("ocio://cg-config-v2.2.0_aces-v1.3_ocio-v2.4");
+		}
+		else if (ConfigurationFile.FilePath == TEXT("ocio://studio-config-latest"))
+		{
+			ConfigurationFile.FilePath = TEXT("ocio://studio-config-v2.2.0_aces-v1.3_ocio-v2.4");
+		}
+	}
+
 	// In non-editor modes, the shader & texture resources are already baked.
 	// For CPU transforms, the config is lazily-loaded.
 	ReloadExistingColorspaces();
diff --git a/Engine/Plugins/Compositing/OpenColorIO/Source/OpenColorIO/Public/OpenColorIOConfiguration.h b/Engine/Plugins/Compositing/OpenColorIO/Source/OpenColorIO/Public/OpenColorIOConfiguration.h
--- a/Engine/Plugins/Compositing/OpenColorIO/Source/OpenColorIO/Public/OpenColorIOConfiguration.h
+++ b/Engine/Plugins/Compositing/OpenColorIO/Source/OpenColorIO/Public/OpenColorIOConfiguration.h
@@ -107,6 +107,7 @@ public:
 
 	//~ Begin UObject interface
 	UE_API virtual void PostInitProperties() override;
+	UE_API virtual void Serialize(FArchive& Ar) override;
 	UE_API virtual void PostLoad() override;
 	UE_API virtual void GetAssetRegistryTags(FAssetRegistryTagsContext RegistryTagsContext) const override;
 #if WITH_EDITORONLY_DATA
diff --git a/Engine/Source/Runtime/Core/Public/UObject/FortniteMainBranchObjectVersions.inl b/Engine/Source/Runtime/Core/Public/UObject/FortniteMainBranchObjectVersions.inl
--- a/Engine/Source/Runtime/Core/Public/UObject/FortniteMainBranchObjectVersions.inl
+++ b/Engine/Source/Runtime/Core/Public/UObject/FortniteMainBranchObjectVersions.inl
@@ -748,5 +748,8 @@ FORTNITE_MAIN_VERSION(VerseRepairDataWithMinusZeroValue, 244)
 // Add new enum property to composite (plugin) UCompositePassColorKeyer to control keying settings
 FORTNITE_MAIN_VERSION(CompositePassKeyerSource, 245)
 
+// OpenColorIO library update to v2.5.1 which transitions from ACES 1.3 to ACES 2.0 defaults
+FORTNITE_MAIN_VERSION(OpenColorIOv251, 246)
+
 // -----<new versions can be added above this line>-------------------------------------------------
 #undef FORTNITE_MAIN_VERSION
