OpenColorIO: Add config type selection menu to facilitate & promote the use of built-ins.

#Synced-CL 50287290

diff --git a/Engine/Plugins/Compositing/OpenColorIO/Source/OpenColorIOEditor/OpenColorIOEditor.Build.cs b/Engine/Plugins/Compositing/OpenColorIO/Source/OpenColorIOEditor/OpenColorIOEditor.Build.cs
--- a/Engine/Plugins/Compositing/OpenColorIO/Source/OpenColorIOEditor/OpenColorIOEditor.Build.cs
+++ b/Engine/Plugins/Compositing/OpenColorIO/Source/OpenColorIOEditor/OpenColorIOEditor.Build.cs
@@ -15,6 +15,7 @@ namespace UnrealBuildTool.Rules
 					"EditorFramework",
 					"Engine",
 					"ImageCore",
+					"InputCore",
 					"LevelEditor",
 					"OpenColorIOWrapper",
 					"OpenColorIO",
diff --git a/Engine/Plugins/Compositing/OpenColorIO/Source/OpenColorIOEditor/Private/OpenColorIOConfigurationCustomization.cpp b/Engine/Plugins/Compositing/OpenColorIO/Source/OpenColorIOEditor/Private/OpenColorIOConfigurationCustomization.cpp
--- a/Engine/Plugins/Compositing/OpenColorIO/Source/OpenColorIOEditor/Private/OpenColorIOConfigurationCustomization.cpp
+++ b/Engine/Plugins/Compositing/OpenColorIO/Source/OpenColorIOEditor/Private/OpenColorIOConfigurationCustomization.cpp
@@ -12,6 +12,8 @@
 #include "PropertyHandle.h"
 #include "Widgets/Input/SButton.h"
 #include "Widgets/Input/SComboButton.h"
+#include "Widgets/Input/SComboBox.h"
+#include "Widgets/Text/STextBlock.h"
 
 #define LOCTEXT_NAMESPACE "OpenColorIOConfigurationCustomization"
 
@@ -26,11 +28,63 @@ void FOpenColorIOConfigurationCustomization::CustomizeDetails(IDetailLayoutBuild
 		return;
 	}
 
-	TSharedRef<IPropertyHandle> ConfigurationFileProperty = DetailBuilder.GetProperty(GET_MEMBER_NAME_CHECKED(UOpenColorIOConfiguration, ConfigurationFile));
+	ConfigTypes.Reset();
+	ConfigTypes.Add(MakeShared<FString>(TEXT("default")));
+	ConfigTypes.Add(MakeShared<FString>(TEXT("cg-config-latest")));
+	ConfigTypes.Add(MakeShared<FString>(TEXT("studio-config-latest")));
+
+	if (UOpenColorIOConfiguration* OpenColorIOConfig = Cast<UOpenColorIOConfiguration>(Objects[0].Get()))
+	{
+		if (FOpenColorIOWrapperConfig* Wrapper = OpenColorIOConfig->GetConfigWrapper())
+		{
+			TArray<FString> ConfigNames = Wrapper->GetBuiltInConfigNames();
+
+			for (const FString& Name : ConfigNames)
+			{
+				ConfigTypes.Add(MakeShared<FString>(Name));
+			}
+		}
+	}
+
+	ConfigTypes.Add(MakeShared<FString>(TEXT("External File")));
+	
+	ConfigurationFilePropertyHandle = DetailBuilder.GetProperty(GET_MEMBER_NAME_CHECKED(UOpenColorIOConfiguration, ConfigurationFile));
+
+	// Update current type
+	OnConfigurationFileChanged();
+
+	ConfigurationFilePropertyHandle->SetOnPropertyValueChanged(FSimpleDelegate::CreateSP(this, &FOpenColorIOConfigurationCustomization::OnConfigurationFileChanged));
+	ConfigurationFilePropertyHandle->SetOnChildPropertyValueChanged(FSimpleDelegate::CreateSP(this, &FOpenColorIOConfigurationCustomization::OnConfigurationFileChanged));
+	ConfigurationFilePropertyHandle->SetOnPropertyResetToDefault(FSimpleDelegate::CreateSP(this, &FOpenColorIOConfigurationCustomization::OnConfigurationFileChanged));
+
 	IDetailCategoryBuilder& ConfigCategory = DetailBuilder.EditCategory("Config");
-	ConfigCategory.AddProperty(DetailBuilder.GetProperty((GET_MEMBER_NAME_CHECKED(UOpenColorIOConfiguration, ConfigurationFile))));
 
-	DetailBuilder.AddCustomRowToCategory(ConfigurationFileProperty, LOCTEXT("Config", "Reload and Rebuild"))
+	const FText TooltipTypes = LOCTEXT("ConfigTypesTooltip", "Types of OpenColorIO configuration, which facilitates the selection of built-ins (latest, or version locked).");
+
+	ConfigCategory.AddCustomRow(LOCTEXT("ConfigTypes", "Types"))
+		.NameContent()
+		[
+			SNew(STextBlock)
+			.Font(IDetailLayoutBuilder::GetDetailFont())
+			.Text(LOCTEXT("ConfigTypeLabel", "Type"))
+			.ToolTipText(TooltipTypes)
+		]
+		.ValueContent()
+		[
+			SAssignNew(TypeCombo, SComboBox<TSharedPtr<FString>>)
+			.OptionsSource(&ConfigTypes)
+			.OnSelectionChanged(this, &FOpenColorIOConfigurationCustomization::OnConfigTypeChanged)
+			.OnGenerateWidget(this, &FOpenColorIOConfigurationCustomization::MakeConfigTypeWidget)
+			.InitiallySelectedItem(CurrentConfigType)
+			[
+				SNew(STextBlock)
+				.Text(this, &FOpenColorIOConfigurationCustomization::GetCurrentConfigTypeLabel)
+			]
+		];
+
+	ConfigCategory.AddProperty(ConfigurationFilePropertyHandle);
+
+	ConfigCategory.AddCustomRow(LOCTEXT("Config", "Reload and Rebuild"))
 		.NameContent()
 		[
 			SNew(STextBlock)
@@ -55,7 +109,7 @@ void FOpenColorIOConfigurationCustomization::CustomizeDetails(IDetailLayoutBuild
 			})
 		];
 
-	TSharedPtr<IPropertyHandle> ConfigurationObjectHandle = ConfigurationFileProperty->GetParentHandle();
+	TSharedPtr<IPropertyHandle> ConfigurationObjectHandle = ConfigurationFilePropertyHandle->GetParentHandle();
 
 	DetailBuilder.RegisterInstancedCustomPropertyTypeLayout(
 			FOpenColorIOColorSpace::StaticStruct()->GetFName(),
@@ -66,4 +120,110 @@ void FOpenColorIOConfigurationCustomization::CustomizeDetails(IDetailLayoutBuild
 		FOnGetPropertyTypeCustomizationInstance::CreateLambda([ConfigurationObjectHandle] { return FOpenColorIODisplayViewCustomization::MakeInstance(ConfigurationObjectHandle); }));
 }
 
+void FOpenColorIOConfigurationCustomization::OnConfigTypeChanged(TSharedPtr<FString> InValue, ESelectInfo::Type InSelectInfo)
+{
+	if (InSelectInfo == ESelectInfo::Direct)
+	{
+		return;
+	}
+
+	CurrentConfigType = InValue;
+	
+	if (ConfigurationFilePropertyHandle.IsValid() && InValue.IsValid())
+	{
+		if (FStructProperty* StructProperty = CastField<FStructProperty>(ConfigurationFilePropertyHandle->GetProperty()))
+		{
+			TArray<void*> RawData;
+			ConfigurationFilePropertyHandle->AccessRawData(RawData);
+
+			FFilePath* PreviousFilePath = nullptr;
+			if (!RawData.IsEmpty())
+			{
+				PreviousFilePath = reinterpret_cast<FFilePath*>(RawData[0]);
+			}
+
+			if (PreviousFilePath)
+			{
+				FFilePath NewFilePath;
+
+				if (InValue->Equals(*ConfigTypes.Last()))
+				{
+					NewFilePath.FilePath = TEXT("");
+				}
+				else
+				{
+					NewFilePath.FilePath = *InValue;
+					NewFilePath.FilePath.InsertAt(0, TEXT("ocio://"));
+				}
+
+				FString TextValue;
+				StructProperty->Struct->ExportText(TextValue, &NewFilePath, PreviousFilePath, nullptr, EPropertyPortFlags::PPF_None, nullptr);
+
+				ensure(ConfigurationFilePropertyHandle->SetValueFromFormattedString(TextValue, EPropertyValueSetFlags::DefaultFlags) == FPropertyAccess::Result::Success);
+			}
+		}
+	}
+}
+
+void FOpenColorIOConfigurationCustomization::OnConfigurationFileChanged()
+{
+	bool bUpdated = false;
+	
+	if (ConfigurationFilePropertyHandle && ConfigurationFilePropertyHandle->IsValidHandle())
+	{
+		FFilePath* CurrentFilePath = nullptr;
+
+		// Update current type
+		TArray<void*> RawData;
+		ConfigurationFilePropertyHandle->AccessRawData(RawData);
+		if (!RawData.IsEmpty())
+		{
+			CurrentFilePath = reinterpret_cast<FFilePath*>(RawData[0]);
+		}
+
+		if (CurrentFilePath != nullptr)
+		{
+			FString CurrentType = CurrentFilePath->FilePath;
+			CurrentType.RemoveFromStart(TEXT("ocio://"));
+
+			const TSharedPtr<FString>* FoundType = ConfigTypes.FindByPredicate([&CurrentType](const TSharedPtr<FString>& InValue) -> bool
+				{
+					return InValue.IsValid() ? InValue->Equals(CurrentType) : false;
+				});
+
+			if (FoundType)
+			{
+				CurrentConfigType = *FoundType;
+				bUpdated = true;
+			}
+		}
+	}
+
+	if (!bUpdated)
+	{
+		CurrentConfigType = ConfigTypes.Last();
+	}
+
+	if (TypeCombo.IsValid())
+	{
+		TypeCombo->SetSelectedItem(CurrentConfigType);
+	}
+}
+
+TSharedRef<SWidget> FOpenColorIOConfigurationCustomization::MakeConfigTypeWidget(TSharedPtr<FString> InOption)
+{
+	return SNew(STextBlock).Text(FText::FromString(*InOption));
+}
+
+FText FOpenColorIOConfigurationCustomization::GetCurrentConfigTypeLabel() const
+{
+	if (CurrentConfigType.IsValid())
+	{
+		return FText::FromString(*CurrentConfigType);
+	}
+
+	return LOCTEXT("Invalid", "Invalid");
+}
+
+
 #undef LOCTEXT_NAMESPACE
diff --git a/Engine/Plugins/Compositing/OpenColorIO/Source/OpenColorIOEditor/Private/OpenColorIOConfigurationCustomization.h b/Engine/Plugins/Compositing/OpenColorIO/Source/OpenColorIOEditor/Private/OpenColorIOConfigurationCustomization.h
--- a/Engine/Plugins/Compositing/OpenColorIO/Source/OpenColorIOEditor/Private/OpenColorIOConfigurationCustomization.h
+++ b/Engine/Plugins/Compositing/OpenColorIO/Source/OpenColorIOEditor/Private/OpenColorIOConfigurationCustomization.h
@@ -5,6 +5,11 @@
 #include "CoreMinimal.h"
 #include "IDetailCustomization.h"
 
+namespace ESelectInfo { enum Type : int; }
+class IPropertyHandle;
+template <typename OptionType> class SComboBox;
+class SWidget;
+class IDetailLayoutBuilder;
 
 /**
  * Implements a details view customization for the UOpenColorIOConfiguration
@@ -18,5 +23,31 @@ public:
 	}
 
 	virtual void CustomizeDetails(IDetailLayoutBuilder& DetailBuilder) override;
+
+private:
+
+	/** Callback for when the config type changes. */
+	void OnConfigTypeChanged(TSharedPtr<FString> InValue, ESelectInfo::Type InSelectInfo);
+
+	/** Callback for when the configuration file path changes. */
+	void OnConfigurationFileChanged();
+
+	/** Create the config type widget. */
+	TSharedRef<SWidget> MakeConfigTypeWidget(TSharedPtr<FString> InOption);
+
+	/** Get the current config type label. */
+	FText GetCurrentConfigTypeLabel() const;
+
+	/** List of available config types */
+	TArray<TSharedPtr<FString>> ConfigTypes;
+
+	/** Currently selected type */
+	TSharedPtr<FString> CurrentConfigType;
+
+	/** Handle to the ConfigurationFile property */
+	TSharedPtr<IPropertyHandle> ConfigurationFilePropertyHandle;
+
+	/** Cached combo box widget. */
+	TSharedPtr<SComboBox<TSharedPtr<FString>>> TypeCombo;
 };
 
diff --git a/Engine/Source/Runtime/OpenColorIOWrapper/Private/OpenColorIOWrapper.cpp b/Engine/Source/Runtime/OpenColorIOWrapper/Private/OpenColorIOWrapper.cpp
--- a/Engine/Source/Runtime/OpenColorIOWrapper/Private/OpenColorIOWrapper.cpp
+++ b/Engine/Source/Runtime/OpenColorIOWrapper/Private/OpenColorIOWrapper.cpp
@@ -436,6 +436,27 @@ FString FOpenColorIOWrapperConfig::GetDebugString() const
 	return DebugStringBuilder.ToString();
 }
 
+TArray<FString> FOpenColorIOWrapperConfig::GetBuiltInConfigNames(bool bRecommendedOnly) const
+{
+	TArray<FString> Names;
+
+	const OCIO_NAMESPACE::BuiltinConfigRegistry& Registry = OCIO_NAMESPACE::BuiltinConfigRegistry::Get();
+
+	for (size_t Index = 0; Index < Registry.getNumBuiltinConfigs(); ++Index)
+	{
+		const ANSICHAR* Name = Registry.getBuiltinConfigName(Index);
+		
+		if (bRecommendedOnly && !Registry.isBuiltinConfigRecommended(Index))
+		{
+			continue;
+		}
+
+		Names.Add(Name);
+	}
+
+	return Names;
+}
+
 
 FOpenColorIOWrapperEngineConfig::FOpenColorIOWrapperEngineConfig()
 	: FOpenColorIOWrapperConfig()
diff --git a/Engine/Source/Runtime/OpenColorIOWrapper/Public/OpenColorIOWrapper.h b/Engine/Source/Runtime/OpenColorIOWrapper/Public/OpenColorIOWrapper.h
--- a/Engine/Source/Runtime/OpenColorIOWrapper/Public/OpenColorIOWrapper.h
+++ b/Engine/Source/Runtime/OpenColorIOWrapper/Public/OpenColorIOWrapper.h
@@ -120,6 +120,9 @@ public:
 	/** Returns a config debug string. */
 	OPENCOLORIOWRAPPER_API FString GetDebugString() const;
 
+	/** Returns all the available built-in config names. */
+	OPENCOLORIOWRAPPER_API TArray<FString> GetBuiltInConfigNames(bool bRecommendedOnly = false) const;
+
 protected:
 
 	TPimplPtr<struct FOpenColorIOConfigPimpl, EPimplPtrMode::DeepCopy> Pimpl;
