OpenColorIO: Adjustment to previous CL to display UI names in new type selection menu instead of the paths.

#Synced-CL 50447978

diff --git a/Engine/Plugins/Compositing/OpenColorIO/Source/OpenColorIOEditor/Private/OpenColorIOConfigurationCustomization.cpp b/Engine/Plugins/Compositing/OpenColorIO/Source/OpenColorIOEditor/Private/OpenColorIOConfigurationCustomization.cpp
--- a/Engine/Plugins/Compositing/OpenColorIO/Source/OpenColorIOEditor/Private/OpenColorIOConfigurationCustomization.cpp
+++ b/Engine/Plugins/Compositing/OpenColorIO/Source/OpenColorIOEditor/Private/OpenColorIOConfigurationCustomization.cpp
@@ -17,7 +17,6 @@
 
 #define LOCTEXT_NAMESPACE "OpenColorIOConfigurationCustomization"
 
-
 void FOpenColorIOConfigurationCustomization::CustomizeDetails(IDetailLayoutBuilder& DetailBuilder)
 {
 	TArray<TWeakObjectPtr<UObject>> Objects;
@@ -28,23 +27,22 @@ void FOpenColorIOConfigurationCustomization::CustomizeDetails(IDetailLayoutBuild
 		return;
 	}
 
-	ConfigTypes.Reset();
-	ConfigTypes.Add(MakeShared<FString>(TEXT("default")));
-	ConfigTypes.Add(MakeShared<FString>(TEXT("cg-config-latest")));
-	ConfigTypes.Add(MakeShared<FString>(TEXT("studio-config-latest")));
-
 	if (UOpenColorIOConfiguration* OpenColorIOConfig = Cast<UOpenColorIOConfiguration>(Objects[0].Get()))
 	{
 		if (FOpenColorIOWrapperConfig* Wrapper = OpenColorIOConfig->GetConfigWrapper())
 		{
-			TArray<FString> ConfigNames = Wrapper->GetBuiltInConfigNames();
-
-			for (const FString& Name : ConfigNames)
-			{
-				ConfigTypes.Add(MakeShared<FString>(Name));
-			}
+			constexpr bool bRecommendedOnly = false;
+			constexpr bool bShortenUIName = true;
+			CachedConfigNamePairs = Wrapper->GetBuiltInConfigNames(bRecommendedOnly, bShortenUIName);
 		}
 	}
+	
+	ConfigTypes.Reset(CachedConfigNamePairs.Num());
+
+	for (const TPair<FString, FString>& Pair : CachedConfigNamePairs)
+	{
+		ConfigTypes.Add(MakeShared<FString>(Pair.Value));
+	}
 
 	ConfigTypes.Add(MakeShared<FString>(TEXT("External File")));
 	
@@ -148,12 +146,20 @@ void FOpenColorIOConfigurationCustomization::OnConfigTypeChanged(TSharedPtr<FStr
 
 				if (InValue->Equals(*ConfigTypes.Last()))
 				{
+					// If external file, we clear the path
 					NewFilePath.FilePath = TEXT("");
 				}
 				else
 				{
-					NewFilePath.FilePath = *InValue;
-					NewFilePath.FilePath.InsertAt(0, TEXT("ocio://"));
+					for (const TPair<FString, FString>& Pair : CachedConfigNamePairs)
+					{
+						if (InValue->Equals(Pair.Value))
+						{
+							NewFilePath.FilePath = Pair.Key;
+							NewFilePath.FilePath.InsertAt(0, TEXT("ocio://"));
+							break;
+						}
+					}
 				}
 
 				FString TextValue;
@@ -183,18 +189,21 @@ void FOpenColorIOConfigurationCustomization::OnConfigurationFileChanged()
 
 		if (CurrentFilePath != nullptr)
 		{
-			FString CurrentType = CurrentFilePath->FilePath;
-			CurrentType.RemoveFromStart(TEXT("ocio://"));
-
-			const TSharedPtr<FString>* FoundType = ConfigTypes.FindByPredicate([&CurrentType](const TSharedPtr<FString>& InValue) -> bool
-				{
-					return InValue.IsValid() ? InValue->Equals(CurrentType) : false;
-				});
+			FString ResolvedPath = FOpenColorIOWrapperConfig::ResolveConfigPath(*CurrentFilePath->FilePath);
+			ResolvedPath.RemoveFromStart(TEXT("ocio://"));
 
-			if (FoundType)
+			if (const FString* UIName = CachedConfigNamePairs.Find(ResolvedPath))
 			{
-				CurrentConfigType = *FoundType;
-				bUpdated = true;
+				const TSharedPtr<FString>* FoundType = ConfigTypes.FindByPredicate([UIName](const TSharedPtr<FString>& InValue) -> bool
+					{
+						return InValue.IsValid() ? InValue->Equals(*UIName) : false;
+					});
+
+				if (FoundType)
+				{
+					CurrentConfigType = *FoundType;
+					bUpdated = true;
+				}
 			}
 		}
 	}
diff --git a/Engine/Plugins/Compositing/OpenColorIO/Source/OpenColorIOEditor/Private/OpenColorIOConfigurationCustomization.h b/Engine/Plugins/Compositing/OpenColorIO/Source/OpenColorIOEditor/Private/OpenColorIOConfigurationCustomization.h
--- a/Engine/Plugins/Compositing/OpenColorIO/Source/OpenColorIOEditor/Private/OpenColorIOConfigurationCustomization.h
+++ b/Engine/Plugins/Compositing/OpenColorIO/Source/OpenColorIOEditor/Private/OpenColorIOConfigurationCustomization.h
@@ -38,6 +38,9 @@ private:
 	/** Get the current config type label. */
 	FText GetCurrentConfigTypeLabel() const;
 
+	/** Cached OpenColorIO built-in name/UI-name pairs. */
+	TMap<FString, FString> CachedConfigNamePairs;
+
 	/** List of available config types */
 	TArray<TSharedPtr<FString>> ConfigTypes;
 
diff --git a/Engine/Source/Runtime/OpenColorIOWrapper/Private/OpenColorIOWrapper.cpp b/Engine/Source/Runtime/OpenColorIOWrapper/Private/OpenColorIOWrapper.cpp
--- a/Engine/Source/Runtime/OpenColorIOWrapper/Private/OpenColorIOWrapper.cpp
+++ b/Engine/Source/Runtime/OpenColorIOWrapper/Private/OpenColorIOWrapper.cpp
@@ -436,27 +436,41 @@ FString FOpenColorIOWrapperConfig::GetDebugString() const
 	return DebugStringBuilder.ToString();
 }
 
-TArray<FString> FOpenColorIOWrapperConfig::GetBuiltInConfigNames(bool bRecommendedOnly) const
+TMap<FString, FString> FOpenColorIOWrapperConfig::GetBuiltInConfigNames(bool bRecommendedOnly, bool bShortenUIName) const
 {
-	TArray<FString> Names;
-
+	TMap<FString, FString> Names;
+	
 	const OCIO_NAMESPACE::BuiltinConfigRegistry& Registry = OCIO_NAMESPACE::BuiltinConfigRegistry::Get();
+	const int32 NumBuiltIns = Registry.getNumBuiltinConfigs();
+	Names.Reserve(NumBuiltIns);
 
-	for (size_t Index = 0; Index < Registry.getNumBuiltinConfigs(); ++Index)
-	{
-		const ANSICHAR* Name = Registry.getBuiltinConfigName(Index);
-		
+	for (size_t Index = 0; Index < NumBuiltIns; ++Index)
+	{		
 		if (bRecommendedOnly && !Registry.isBuiltinConfigRecommended(Index))
 		{
 			continue;
 		}
 
-		Names.Add(Name);
+		const ANSICHAR* Name = Registry.getBuiltinConfigName(Index);
+		FString UIName = Registry.getBuiltinConfigUIName(Index);
+		
+		if (bShortenUIName)
+		{
+			UIName.RemoveFromStart(TEXT("Academy Color Encoding System - "));
+		}
+
+		Names.Emplace(Name, MoveTemp(UIName));
 	}
 
+
 	return Names;
 }
 
+FString FOpenColorIOWrapperConfig::ResolveConfigPath(const TCHAR* InOriginalPath) noexcept
+{
+	return OCIO_NAMESPACE::ResolveConfigPath(TCHAR_TO_ANSI(InOriginalPath));
+}
+
 
 FOpenColorIOWrapperEngineConfig::FOpenColorIOWrapperEngineConfig()
 	: FOpenColorIOWrapperConfig()
diff --git a/Engine/Source/Runtime/OpenColorIOWrapper/Public/OpenColorIOWrapper.h b/Engine/Source/Runtime/OpenColorIOWrapper/Public/OpenColorIOWrapper.h
--- a/Engine/Source/Runtime/OpenColorIOWrapper/Public/OpenColorIOWrapper.h
+++ b/Engine/Source/Runtime/OpenColorIOWrapper/Public/OpenColorIOWrapper.h
@@ -120,8 +120,11 @@ public:
 	/** Returns a config debug string. */
 	OPENCOLORIOWRAPPER_API FString GetDebugString() const;
 
-	/** Returns all the available built-in config names. */
-	OPENCOLORIOWRAPPER_API TArray<FString> GetBuiltInConfigNames(bool bRecommendedOnly = false) const;
+	/** Returns all the available built-in config name/UI name pairs. */
+	OPENCOLORIOWRAPPER_API TMap<FString, FString> GetBuiltInConfigNames(bool bRecommendedOnly, bool bShortenUIName) const;
+
+	/** Make a config path forward-compatible by replacing special built-in config names with the current name. */
+	OPENCOLORIOWRAPPER_API static FString ResolveConfigPath(const TCHAR* InOriginalPath) noexcept;
 
 protected:
 
