OCIO: Fix for linear floating point (SDR) backbuffer.

Colors must remain linear with a DXGI_FORMAT_R16G16B16A16_FLOAT swapchain (including transform output).

#Synced-CL 51563080

diff --git a/Engine/Plugins/Compositing/OpenColorIO/Source/OpenColorIO/Private/OpenColorIORendering.cpp b/Engine/Plugins/Compositing/OpenColorIO/Source/OpenColorIO/Private/OpenColorIORendering.cpp
--- a/Engine/Plugins/Compositing/OpenColorIO/Source/OpenColorIO/Private/OpenColorIORendering.cpp
+++ b/Engine/Plugins/Compositing/OpenColorIO/Source/OpenColorIO/Private/OpenColorIORendering.cpp
@@ -8,6 +8,7 @@
 #include "Engine/World.h"
 #include "EngineModule.h"
 #include "GlobalShader.h"
+#include "HDRHelper.h"
 #include "Logging/LogMacros.h"
 #include "Logging/MessageLog.h"
 #include "OpenColorIOColorSpace.h"
@@ -150,9 +151,18 @@ void FOpenColorIORendering::AddPass_RenderThread(
 {
 	const float EngineDisplayGamma = View.Family->RenderTarget->GetDisplayGamma();
 	const float DefaultDisplayGammaRT = UTextureRenderTarget::GetDefaultDisplayGamma();
+	float DisplayGamma = DefaultDisplayGammaRT / EngineDisplayGamma;
+
+	if (IsDefaultBackBufferLinearSDR())
+	{
+		// Note: To match editor in this mode, transforms need to output linear colors.
+		DisplayGamma = 1.0f;
+	}
 	// There is a special case where post processing and tonemapper are disabled. In this case tonemapper applies a static display Inverse of Gamma which defaults to 2.2.
-	// In the case when Both PostProcessing and ToneMapper are disabled we apply gamma manually. In every other case we apply inverse gamma before applying OCIO.
-	float DisplayGamma = (View.Family->EngineShowFlags.Tonemapper == 0) || (View.Family->EngineShowFlags.PostProcessing == 0) ? DefaultDisplayGammaRT : DefaultDisplayGammaRT / EngineDisplayGamma;
+	else if ((View.Family->EngineShowFlags.Tonemapper == 0) || (View.Family->EngineShowFlags.PostProcessing == 0))
+	{
+		DisplayGamma = DefaultDisplayGammaRT;
+	}
 
 	FOpenColorIORendering::AddPass_RenderThread(
 		GraphBuilder,
